Move replication UI to ConcertClientSharedSlate
- Add UMultiUserReplicationSessionPreset
- Add UI for editing properties to Multi User active session UI in Replication tab (does not affect replication right now)

diff --git a/Engine/Plugins/Developer/Concert/ConcertApp/MultiUserClient/MultiUserClient.uplugin b/Engine/Plugins/Developer/Concert/ConcertApp/MultiUserClient/MultiUserClient.uplugin
--- a/Engine/Plugins/Developer/Concert/ConcertApp/MultiUserClient/MultiUserClient.uplugin
+++ b/Engine/Plugins/Developer/Concert/ConcertApp/MultiUserClient/MultiUserClient.uplugin
@@ -25,11 +25,6 @@
 			"Type": "Runtime",
 			"LoadingPhase": "Default"
 		},
-		{
-			"Name": "MultiUserReplication",
-			"Type": "Runtime",
-			"LoadingPhase": "Default"
-		},
 		{
 			"Name": "MultiUserReplicationEditor",
 			"Type": "Editor",
@@ -41,6 +36,10 @@
 			"Name": "ConcertSyncClient",
 			"Enabled": true
 		},
+		{
+			"Name": "ConcertClientSharedSlate",
+			"Enabled": true
+		},
         {
           "Name": "ConcertSharedSlate",
           "Enabled": true
diff --git a/Engine/Plugins/Developer/Concert/ConcertApp/MultiUserClient/Source/MultiUserClient/MultiUserClient.Build.cs b/Engine/Plugins/Developer/Concert/ConcertApp/MultiUserClient/Source/MultiUserClient/MultiUserClient.Build.cs
--- a/Engine/Plugins/Developer/Concert/ConcertApp/MultiUserClient/Source/MultiUserClient/MultiUserClient.Build.cs
+++ b/Engine/Plugins/Developer/Concert/ConcertApp/MultiUserClient/Source/MultiUserClient/MultiUserClient.Build.cs
@@ -17,26 +17,28 @@ namespace UnrealBuildTool.Rules
 			PrivateDependencyModuleNames.AddRange(
 				new string[]
 				{
+					// Engine
 					"ApplicationCore",
-					"Concert",
-					"ConcertClient",
-					"ConcertClientSharedSlate",
-					"ConcertSharedSlate",
-					"ConcertSyncClient",
-					"ConcertSyncCore",
-					"ConcertTransport",
 					"ContentBrowser",
 					"DesktopPlatform",
 					"EditorStyle",
 					"InputCore",
 					"MessageLog",
-					"MultiUserReplication",
 					"Projects",
 					"Slate",
 					"SlateCore",
 					"SourceControl",
 					"ToolMenus",
-					"ToolWidgets"
+					"ToolWidgets",
+					
+					// Concert
+					"Concert",
+					"ConcertClient",
+					"ConcertClientSharedSlate",
+					"ConcertSharedSlate",
+					"ConcertSyncClient",
+					"ConcertSyncCore",
+					"ConcertTransport",
 				}
 			);
 
diff --git a/Engine/Plugins/Developer/Concert/ConcertApp/MultiUserClient/Source/MultiUserClient/Private/Replication/MultiUserReplicationManager.cpp b/Engine/Plugins/Developer/Concert/ConcertApp/MultiUserClient/Source/MultiUserClient/Private/Replication/MultiUserReplicationManager.cpp
--- a/Engine/Plugins/Developer/Concert/ConcertApp/MultiUserClient/Source/MultiUserClient/Private/Replication/MultiUserReplicationManager.cpp
+++ b/Engine/Plugins/Developer/Concert/ConcertApp/MultiUserClient/Source/MultiUserClient/Private/Replication/MultiUserReplicationManager.cpp
@@ -3,37 +3,73 @@
 #include "MultiUserReplicationManager.h"
 
 #include "IConcertSyncClient.h"
-#include "MultiUserReplicationClientProfileAsset.h"
 #include "ReplicationUtils.h"
 
+#include "UObject/Package.h"
+
 namespace UE::MultiUserClient
 {
 	FMultiUserReplicationManager::FMultiUserReplicationManager(TSharedRef<IConcertSyncClient> InClient)
 		: Client(MoveTemp(InClient))
-	{}
+		, SessionContent(NewObject<UMultiUserReplicationSessionPreset>(GetTransientPackage(), NAME_None, RF_Transient))
+		, LocalClientContent(SessionContent->AddClient())
+	{
+		Client->GetConcertClient()->OnSessionConnectionChanged().AddRaw(
+			this,
+			&FMultiUserReplicationManager::OnSessionConnectionChanged
+			);
+	}
 
-	void FMultiUserReplicationManager::JoinReplicationSession(const UMultiUserReplicationClientProfileAsset& Asset)
+	FMultiUserReplicationManager::~FMultiUserReplicationManager()
 	{
-		Replication::JoinSessionForMultiUser(Client, Asset.ToJoinArgs());
+		Client->GetConcertClient()->OnSessionConnectionChanged().RemoveAll(this);
 	}
 
-	void FMultiUserReplicationManager::LeaveSession()
+	void FMultiUserReplicationManager::AddReferencedObjects(FReferenceCollector& Collector)
 	{
-		if (IConcertClientReplicationManager* ReplicationManager = Client->GetReplicationManager())
+		Collector.AddReferencedObject(SessionContent);
+		Collector.AddReferencedObject(LocalClientContent);
+	}
+
+	void FMultiUserReplicationManager::OnSessionConnectionChanged(
+		IConcertClientSession& ConcertClientSession,
+		EConcertConnectionStatus ConcertConnectionStatus
+		)
+	{
+		switch (ConcertConnectionStatus)
 		{
-			ReplicationManager->LeaveReplicationSession();
+		case EConcertConnectionStatus::Connecting:
+			break;
+		case EConcertConnectionStatus::Connected:
+			OnJoinSession(ConcertClientSession);
+			break;
+		case EConcertConnectionStatus::Disconnecting:
+			break;
+		case EConcertConnectionStatus::Disconnected:
+			OnLeaveSession(ConcertClientSession);
+			break;
+		default: ;
 		}
 	}
 
-	bool FMultiUserReplicationManager::CanJoin() const
+	void FMultiUserReplicationManager::OnJoinSession(IConcertClientSession& ConcertClientSession)
 	{
-		IConcertClientReplicationManager* ReplicationManager = Client->GetReplicationManager();
-		return ReplicationManager && ReplicationManager->CanJoin();
+		Replication::JoinSessionForMultiUser(Client, {});
+	}
+
+	void FMultiUserReplicationManager::OnLeaveSession(IConcertClientSession& ConcertClientSession)
+	{
+		ClearSessionData();
+		
+		if (IConcertClientReplicationManager* ReplicationManager = Client->GetReplicationManager())
+		{
+			ReplicationManager->LeaveReplicationSession();
+		}
 	}
 
-	bool FMultiUserReplicationManager::IsConnectedToReplicationSession() const
+	void FMultiUserReplicationManager::ClearSessionData()
 	{
-		IConcertClientReplicationManager* ReplicationManager = Client->GetReplicationManager();
-		return ReplicationManager && ReplicationManager->IsConnectedToReplicationSession();
+		SessionContent->ClearClients();
+		LocalClientContent = SessionContent->AddClient();
 	}
 }
diff --git a/Engine/Plugins/Developer/Concert/ConcertApp/MultiUserClient/Source/MultiUserClient/Private/Replication/MultiUserReplicationManager.h b/Engine/Plugins/Developer/Concert/ConcertApp/MultiUserClient/Source/MultiUserClient/Private/Replication/MultiUserReplicationManager.h
--- a/Engine/Plugins/Developer/Concert/ConcertApp/MultiUserClient/Source/MultiUserClient/Private/Replication/MultiUserReplicationManager.h
+++ b/Engine/Plugins/Developer/Concert/ConcertApp/MultiUserClient/Source/MultiUserClient/Private/Replication/MultiUserReplicationManager.h
@@ -2,11 +2,16 @@
 
 #pragma once
 
+#include "Assets/MultiUserReplicationSessionPreset.h"
+#include "IConcertSession.h"
 #include "Templates/SharedPointer.h"
+#include "UObject/GCObject.h"
 
-class UMultiUserReplicationClientProfileAsset;
+class IConcertClientSession;
 class IConcertSyncClient;
 
+enum class EConcertConnectionStatus : uint8;
+
 namespace UE::MultiUserClient
 {
 	/**
@@ -19,26 +24,39 @@ namespace UE::MultiUserClient
 	 *
 	 * This class implements the Fence design pattern. All knowledge Multi-User might need should be encapsulated by this class.
 	 */
-	class FMultiUserReplicationManager
+	class FMultiUserReplicationManager : public FGCObject
 	{
 	public:
-
+		
 		FMultiUserReplicationManager(TSharedRef<IConcertSyncClient> InClient);
+		virtual ~FMultiUserReplicationManager() override;
 
-		/** Joins a replication session. */
-		void JoinReplicationSession(const UMultiUserReplicationClientProfileAsset& Asset);
-		/** Leaves the current replication session */
-		void LeaveSession();
+		UMultiUserReplicationSessionPreset* GetSessionContent() const { return SessionContent; }
+		UMultiUserReplicationClientPreset* GetLocalClientContent() const { return LocalClientContent; }
 
-		/** Whether it is valid to call JoinReplicationSession. */
-		bool CanJoin() const;
-		/** Whether it is valid to call LeaveSession. */
-		bool IsConnectedToReplicationSession() const;
+		//~ Begin FGCObject Interface
+		virtual void AddReferencedObjects(FReferenceCollector& Collector) override;
+		virtual FString GetReferencerName() const override { return TEXT("FMultiUserReplicationManager"); }
+		//~ End FGCObject Interface
 
 	private:
 
 		/** Client through which the replication bridge is accessed. */
 		TSharedRef<IConcertSyncClient> Client;
+
+		/** The state of the server is synched up with this object and displayed in the UI. */
+		TObjectPtr<UMultiUserReplicationSessionPreset> SessionContent;
+		/** Data for the local client. Also part of SessionContent->ClientPresets. */
+		TObjectPtr<UMultiUserReplicationClientPreset> LocalClientContent;
+		
+		void OnSessionConnectionChanged(IConcertClientSession& ConcertClientSession, EConcertConnectionStatus ConcertConnectionStatus);
+		
+		/** Joins a replication session. */
+		void OnJoinSession(IConcertClientSession& ConcertClientSession);
+		/** Leaves the current replication session */
+		void OnLeaveSession(IConcertClientSession& ConcertClientSession);
+		
+		void ClearSessionData();
 	};
 }
 
diff --git a/Engine/Plugins/Developer/Concert/ConcertApp/MultiUserClient/Source/MultiUserClient/Private/Widgets/ActiveSession/Replication/SReplicationControlsTab.cpp b/Engine/Plugins/Developer/Concert/ConcertApp/MultiUserClient/Source/MultiUserClient/Private/Widgets/ActiveSession/Replication/SReplicationControlsTab.cpp
--- a/Engine/Plugins/Developer/Concert/ConcertApp/MultiUserClient/Source/MultiUserClient/Private/Widgets/ActiveSession/Replication/SReplicationControlsTab.cpp
+++ b/Engine/Plugins/Developer/Concert/ConcertApp/MultiUserClient/Source/MultiUserClient/Private/Widgets/ActiveSession/Replication/SReplicationControlsTab.cpp
@@ -2,8 +2,11 @@
 
 #include "SReplicationControlsTab.h"
 
-#include "ConcertFrontendUtils.h"
-#include "SReplicationConnectionControls.h"
+#include "Assets/MultiUserReplicationSessionPreset.h"
+#include "Replication/MultiUserReplicationManager.h"
+#include "Replication/ReplicationWidgetFactories.h"
+#include "Replication/Editor/Model/Object/EditorObjectSelectionSourceModel.h"
+#include "Replication/Editor/Model/Property/SelectPropertyFromUClassModel.h"
 
 #include "Styling/AppStyle.h"
 #include "Widgets/Layout/SExpandableArea.h"
@@ -15,102 +18,21 @@
 
 namespace UE::MultiUserClient
 {
-	namespace Private
-	{
-		static TSharedRef<SWidget> MakeExpandableArea(
-			TSharedPtr<SExpandableArea>& OutArea,
-			FOnBooleanValueChanged OnAreaExpansionChanged,
-			FText Label,
-			FText Tooltip,
-			TSharedRef<SWidget> BodyContent,
-			bool bInitiallyCollapsed
-			)
-		{
-			return SNew(SBorder)
-				.BorderImage(FAppStyle::Get().GetBrush("ToolPanel.GroupBorder"))
-				.Padding(0.0f)
-				[
-					SAssignNew(OutArea, SExpandableArea)
-					.BorderBackgroundColor(FLinearColor(0.6f, 0.6f, 0.6f, 1.0f))
-					.BorderImage_Lambda([&OutArea]() { return ConcertFrontendUtils::GetExpandableAreaBorderImage(*OutArea); })
-					.BodyBorderImage(FAppStyle::Get().GetBrush("ToolPanel.GroupBorder"))
-					.BodyBorderBackgroundColor(FLinearColor::White)
-					.OnAreaExpansionChanged(MoveTemp(OnAreaExpansionChanged))
-					.Padding(0.0f)
-					.InitiallyCollapsed(bInitiallyCollapsed)
-					.HeaderContent()
-					[
-						SNew(STextBlock)
-						.Text(MoveTemp(Label))
-						.ToolTipText(MoveTemp(Tooltip))
-						.Font(FAppStyle::Get().GetFontStyle("DetailsView.CategoryFontStyle"))
-						.ShadowOffset(FVector2D(1.0f, 1.0f))
-					]
-					.BodyContent()
-					[
-						MoveTemp(BodyContent)
-					]
-				];
-		}
-	}
-	
 	void SReplicationControlsTab::Construct(const FArguments& InArgs, TSharedRef<FMultiUserReplicationManager> InReplicationManager)
 	{
-		constexpr bool bCollpaseInitially_ConnectionArea	= false;
-		constexpr bool bCollpaseInitially_AttributesArea	= true;
-		constexpr bool bCollpaseInitially_StreamsArea		= true;
-		
+		const UMultiUserReplicationClientPreset* LocalClient = InReplicationManager->GetLocalClientContent();
+		check(LocalClient && LocalClient->Stream);
+
+		using namespace ConcertClientSharedSlate;
+		const FCreateEditorParams ReplicationEditorCreationParams
+		{
+			CreatePropertySelectionModel(*LocalClient->Stream, LocalClient->Stream->MakeReplicationMapGetterAttribute()),
+			MakeShared<FEditorObjectSelectionSourceModel>(),
+			MakeShared<FSelectPropertyFromUClassModel>()
+		};
 		ChildSlot
 		[
-			SNew(SSplitter)
-			.Orientation(Orient_Vertical)
-
-			// Connection area
-			+SSplitter::Slot()
-			.SizeRule(TAttribute<SSplitter::ESizeRule>(this, &SReplicationControlsTab::GetConnectionAreaSizeRule))
-			.Value(0.3)
-			[
-				Private::MakeExpandableArea(
-					ConnectionArea,
-					FOnBooleanValueChanged::CreateSP(this, &SReplicationControlsTab::OnConnectionAreaExpansionChanged),
-					LOCTEXT("ConnectionArea.Label", "Connection"),
-					LOCTEXT("ConnectionArea.Tooltip", "Control the replication connection, e.g. join and leave"),
-					SNew(SReplicationConnectionControls, InReplicationManager),
-					bCollpaseInitially_ConnectionArea
-					)
-			]
-			
-			// Attributes area
-			+SSplitter::Slot()
-			.SizeRule(TAttribute<SSplitter::ESizeRule>(this, &SReplicationControlsTab::GetClientAttributesAreaSizeRule))
-			.Value(0.3)
-			[
-				Private::MakeExpandableArea(
-					ClientAttributesArea,
-					FOnBooleanValueChanged::CreateSP(this, &SReplicationControlsTab::OnAClientsttributesAreaExpansionChanged),
-					LOCTEXT("AttributesArea.Label", "Client Receive Attributes"),
-					LOCTEXT("AttributesArea.Tooltip", "View and edit your client attributes, which determine the data received."),
-					SNew(STextBlock)
-						.Text(LOCTEXT("ToDo", "ToDo")),
-					bCollpaseInitially_AttributesArea
-					)
-			]
-			
-			// Streams area
-			+SSplitter::Slot()
-			.SizeRule(TAttribute<SSplitter::ESizeRule>(this, &SReplicationControlsTab::GetStreamsAreaSizeRule))
-			.Value(0.3)
-			[
-				Private::MakeExpandableArea(
-					StreamsArea,
-					FOnBooleanValueChanged::CreateSP(this, &SReplicationControlsTab::OnStreamsAreaExpansionChanged),
-					LOCTEXT("StreamsArea.Label", "Streams"),
-					LOCTEXT("StreamsArea.Tooltip", "View your replication streams as well as toggle sent objects and properties"),
-					SNew(STextBlock)
-						.Text(LOCTEXT("ToDo", "ToDo")),
-					bCollpaseInitially_StreamsArea
-					)
-			]
+			CreateEditor(ReplicationEditorCreationParams)
 		];
 	}
 }
diff --git a/Engine/Plugins/Developer/Concert/ConcertApp/MultiUserClient/Source/MultiUserClient/Private/Widgets/ActiveSession/Replication/SReplicationControlsTab.h b/Engine/Plugins/Developer/Concert/ConcertApp/MultiUserClient/Source/MultiUserClient/Private/Widgets/ActiveSession/Replication/SReplicationControlsTab.h
--- a/Engine/Plugins/Developer/Concert/ConcertApp/MultiUserClient/Source/MultiUserClient/Private/Widgets/ActiveSession/Replication/SReplicationControlsTab.h
+++ b/Engine/Plugins/Developer/Concert/ConcertApp/MultiUserClient/Source/MultiUserClient/Private/Widgets/ActiveSession/Replication/SReplicationControlsTab.h
@@ -12,6 +12,7 @@ namespace UE::MultiUserClient
 {
 	class FMultiUserReplicationManager;
 
+	/** Root widget for replication in Multi-User session. */
 	class SReplicationControlsTab : public SCompoundWidget
 	{
 	public:
@@ -21,31 +22,6 @@ namespace UE::MultiUserClient
 		SLATE_END_ARGS()
 
 		void Construct(const FArguments& InArgs, TSharedRef<FMultiUserReplicationManager> InReplicationManager);
-
-	private:
-
-		/** In this area the user can control the replication connection, e.g. join and leave. */
-		TSharedPtr<SExpandableArea> ConnectionArea;
-		/** In this area the user can view and edit stream attributes, which will be replicated back to the server. */
-		TSharedPtr<SExpandableArea> ClientAttributesArea;
-		/** In this area the user can view the replication streams as well as toggle sent objects and its properties. */
-		TSharedPtr<SExpandableArea> StreamsArea;
-
-		bool bConnectionAreaExpanded = false;
-		bool bAttributesAreaExpanded = false;
-		bool bStreamsAreaExpanded = false;
-
-		/** Handles how much space the 'ConnectionArea' area uses with respect to its expansion state. */
-		SSplitter::ESizeRule GetConnectionAreaSizeRule() const { return SSplitter::ESizeRule::SizeToContent; }
-		void OnConnectionAreaExpansionChanged(bool bExpanded) { bConnectionAreaExpanded = bExpanded; }
-
-		/** Handles how much space the 'ClientAttributesArea' area uses with respect to its expansion state. */
-		SSplitter::ESizeRule GetClientAttributesAreaSizeRule() const { return bAttributesAreaExpanded ? SSplitter::ESizeRule::FractionOfParent : SSplitter::ESizeRule::SizeToContent; }
-		void OnAClientsttributesAreaExpansionChanged(bool bExpanded) { bAttributesAreaExpanded = bExpanded; }
-		
-		/** Handles how much space the 'StreamsArea' area uses with respect to its expansion state. */
-		SSplitter::ESizeRule GetStreamsAreaSizeRule() const { return bStreamsAreaExpanded ? SSplitter::ESizeRule::FractionOfParent : SSplitter::ESizeRule::SizeToContent; }
-		void OnStreamsAreaExpansionChanged(bool bExpanded) { bStreamsAreaExpanded = bExpanded; }
 	};
 }
 
diff --git a/Engine/Plugins/Developer/Concert/ConcertApp/MultiUserClient/Source/MultiUserReplicationEditor/Internal/Assets/MultiUserReplicationClientPreset.h b/Engine/Plugins/Developer/Concert/ConcertApp/MultiUserClient/Source/MultiUserReplicationEditor/Internal/Assets/MultiUserReplicationClientPreset.h
new file mode 100644
--- /dev/null
+++ b/Engine/Plugins/Developer/Concert/ConcertApp/MultiUserClient/Source/MultiUserReplicationEditor/Internal/Assets/MultiUserReplicationClientPreset.h
@@ -0,0 +1,30 @@
+// Copyright Epic Games, Inc. All Rights Reserved.
+
+#pragma once
+
+#include "UObject/Object.h"
+#include "MultiUserReplicationStream.h"
+#include "Replication/Data/ReplicationStreamDescription.h"
+#include "MultiUserReplicationClientPreset.generated.h"
+
+/**
+ * 
+ */
+UCLASS()
+class MULTIUSERREPLICATIONEDITOR_API UMultiUserReplicationClientPreset : public UObject
+{
+	GENERATED_BODY()
+public:
+
+	/** The stream this client is managing */
+	UPROPERTY(Instanced)
+	TObjectPtr<UMultiUserReplicationStream> Stream;
+
+	UMultiUserReplicationClientPreset();
+
+	/** Resets the contents everything to defaults. Empties the stream */
+	void ClearClient();
+
+	/** Generates a description that can be sent to the MU server. */
+	FReplicationStreamDescription GenerateDescription() const;
+};
diff --git a/Engine/Plugins/Developer/Concert/ConcertApp/MultiUserClient/Source/MultiUserReplicationEditor/Internal/Assets/MultiUserReplicationSessionPreset.h b/Engine/Plugins/Developer/Concert/ConcertApp/MultiUserClient/Source/MultiUserReplicationEditor/Internal/Assets/MultiUserReplicationSessionPreset.h
new file mode 100644
--- /dev/null
+++ b/Engine/Plugins/Developer/Concert/ConcertApp/MultiUserClient/Source/MultiUserReplicationEditor/Internal/Assets/MultiUserReplicationSessionPreset.h
@@ -0,0 +1,32 @@
+// Copyright Epic Games, Inc. All Rights Reserved.
+
+#pragma once
+
+#include "MultiUserReplicationClientPreset.h"
+#include "UObject/Object.h"
+#include "MultiUserReplicationSessionPreset.generated.h"
+
+/**
+ * 
+ */
+UCLASS()
+class MULTIUSERREPLICATIONEDITOR_API UMultiUserReplicationSessionPreset : public UObject
+{
+	GENERATED_UCLASS_BODY()
+public:
+
+	UMultiUserReplicationClientPreset* AddClient();
+	void ClearClients();
+
+	const TArray<TObjectPtr<UMultiUserReplicationClientPreset>>& GetClientPresets() const { return ClientPresets; }
+	UMultiUserReplicationClientPreset* GetUnassignedClient() const { return UnassignedClient; }
+
+private:
+	
+	UPROPERTY(Instanced)
+	TArray<TObjectPtr<UMultiUserReplicationClientPreset>> ClientPresets;
+
+	/** Special preset that is assigned to an invalid client - objects without ownership go here. This is not in ClientPresets. */
+	UPROPERTY(Instanced)
+	TObjectPtr<UMultiUserReplicationClientPreset> UnassignedClient;
+};
diff --git a/Engine/Plugins/Developer/Concert/ConcertApp/MultiUserClient/Source/MultiUserReplicationEditor/Internal/Assets/MultiUserReplicationStream.h b/Engine/Plugins/Developer/Concert/ConcertApp/MultiUserClient/Source/MultiUserReplicationEditor/Internal/Assets/MultiUserReplicationStream.h
new file mode 100644
--- /dev/null
+++ b/Engine/Plugins/Developer/Concert/ConcertApp/MultiUserClient/Source/MultiUserReplicationEditor/Internal/Assets/MultiUserReplicationStream.h
@@ -0,0 +1,36 @@
+// Copyright Epic Games, Inc. All Rights Reserved.
+
+#pragma once
+
+#include "UObject/Object.h"
+#include "Replication/Data/ReplicationStreamDescription.h"
+#include "Misc/Attribute.h"
+#include "MultiUserReplicationStream.generated.h"
+
+/** Wraps FObjectReplicationMap so its edition can be transacted in the editor. */
+UCLASS()
+class MULTIUSERREPLICATIONEDITOR_API UMultiUserReplicationStream : public UObject
+{
+	GENERATED_BODY()
+public:
+
+	/** The ID of the stream. Set to 0 for CDO to avoid issues with delta serialization.*/
+	UPROPERTY()
+	FGuid StreamId;
+
+	/** The objects this stream will modify. */
+	UPROPERTY()
+	FObjectReplicationMap ReplicationMap;
+
+	UMultiUserReplicationStream();
+	
+	FReplicationStreamDescription GenerateDescription() const;
+	
+	TAttribute<FObjectReplicationMap*> MakeReplicationMapGetterAttribute()
+	{
+		return TAttribute<FObjectReplicationMap*>::CreateLambda([WeakThis = TWeakObjectPtr<UMultiUserReplicationStream>(this)]()
+		{
+			return WeakThis.IsValid() ? &WeakThis->ReplicationMap : nullptr;
+		});
+	}
+};
diff --git a/Engine/Plugins/Developer/Concert/ConcertApp/MultiUserClient/Source/MultiUserReplicationEditor/MultiUserReplicationEditor.Build.cs b/Engine/Plugins/Developer/Concert/ConcertApp/MultiUserClient/Source/MultiUserReplicationEditor/MultiUserReplicationEditor.Build.cs
--- a/Engine/Plugins/Developer/Concert/ConcertApp/MultiUserClient/Source/MultiUserReplicationEditor/MultiUserReplicationEditor.Build.cs
+++ b/Engine/Plugins/Developer/Concert/ConcertApp/MultiUserClient/Source/MultiUserReplicationEditor/MultiUserReplicationEditor.Build.cs
@@ -6,6 +6,8 @@ namespace UnrealBuildTool.Rules
 	{
 		public MultiUserReplicationEditor(ReadOnlyTargetRules Target) : base(Target)
 		{
+			ShortName = "MURepEditor";
+			
 			PublicDependencyModuleNames.AddRange(
 				new string[]
 				{
@@ -15,8 +17,8 @@ namespace UnrealBuildTool.Rules
 					"Engine",
 					
 					// Concert
-					"ConcertSyncCore",
-					"MultiUserReplication"
+					"ConcertClientSharedSlate",
+					"ConcertSyncCore"
 				}
 			);
 			
diff --git a/Engine/Plugins/Developer/Concert/ConcertApp/MultiUserClient/Source/MultiUserReplicationEditor/Private/Assets/MultiUserReplicationClientPreset.cpp b/Engine/Plugins/Developer/Concert/ConcertApp/MultiUserClient/Source/MultiUserReplicationEditor/Private/Assets/MultiUserReplicationClientPreset.cpp
new file mode 100644
--- /dev/null
+++ b/Engine/Plugins/Developer/Concert/ConcertApp/MultiUserClient/Source/MultiUserReplicationEditor/Private/Assets/MultiUserReplicationClientPreset.cpp
@@ -0,0 +1,22 @@
+// Copyright Epic Games, Inc. All Rights Reserved.
+
+#include "Assets/MultiUserReplicationClientPreset.h"
+
+#include "Assets/MultiUserReplicationStream.h"
+#include "Replication/Data/ReplicationStreamDescription.h"
+
+UMultiUserReplicationClientPreset::UMultiUserReplicationClientPreset()
+{
+	Stream = CreateDefaultSubobject<UMultiUserReplicationStream>(TEXT("ReplicationList"));
+	Stream->SetFlags(RF_Transactional);
+}
+
+void UMultiUserReplicationClientPreset::ClearClient()
+{
+	Stream->ReplicationMap.ReplicatedObjects.Empty();
+}
+
+FReplicationStreamDescription UMultiUserReplicationClientPreset::GenerateDescription() const
+{
+	return Stream->GenerateDescription();
+}
diff --git a/Engine/Plugins/Developer/Concert/ConcertApp/MultiUserClient/Source/MultiUserReplicationEditor/Private/Assets/MultiUserReplicationSessionPreset.cpp b/Engine/Plugins/Developer/Concert/ConcertApp/MultiUserClient/Source/MultiUserReplicationEditor/Private/Assets/MultiUserReplicationSessionPreset.cpp
new file mode 100644
--- /dev/null
+++ b/Engine/Plugins/Developer/Concert/ConcertApp/MultiUserClient/Source/MultiUserReplicationEditor/Private/Assets/MultiUserReplicationSessionPreset.cpp
@@ -0,0 +1,21 @@
+// Copyright Epic Games, Inc. All Rights Reserved.
+
+#include "Assets/MultiUserReplicationSessionPreset.h"
+
+UMultiUserReplicationSessionPreset::UMultiUserReplicationSessionPreset(const FObjectInitializer& ObjectInitializer)
+{
+	UnassignedClient = ObjectInitializer.CreateDefaultSubobject<UMultiUserReplicationClientPreset>(this, TEXT("UnassignedClient"));
+}
+
+UMultiUserReplicationClientPreset* UMultiUserReplicationSessionPreset::AddClient()
+{
+	UMultiUserReplicationClientPreset* Result = NewObject<UMultiUserReplicationClientPreset>(this);
+	ClientPresets.Add(Result);
+	return Result;
+}
+
+void UMultiUserReplicationSessionPreset::ClearClients()
+{
+	ClientPresets.Empty();
+	UnassignedClient->ClearClient();
+}
diff --git a/Engine/Plugins/Developer/Concert/ConcertApp/MultiUserClient/Source/MultiUserReplicationEditor/Private/Assets/MultiUserReplicationStream.cpp b/Engine/Plugins/Developer/Concert/ConcertApp/MultiUserClient/Source/MultiUserReplicationEditor/Private/Assets/MultiUserReplicationStream.cpp
new file mode 100644
--- /dev/null
+++ b/Engine/Plugins/Developer/Concert/ConcertApp/MultiUserClient/Source/MultiUserReplicationEditor/Private/Assets/MultiUserReplicationStream.cpp
@@ -0,0 +1,18 @@
+// Copyright Epic Games, Inc. All Rights Reserved.
+
+#include "Assets/MultiUserReplicationStream.h"
+
+#include "Replication/Data/ReplicationStreamDescription.h"
+
+UMultiUserReplicationStream::UMultiUserReplicationStream()
+{
+	// To avoid issues with delta serialization the CDO is always 0
+	StreamId = HasAnyFlags(RF_ClassDefaultObject)
+		? FGuid{}
+		: FGuid::NewGuid();
+}
+
+FReplicationStreamDescription UMultiUserReplicationStream::GenerateDescription() const
+{
+	return { StreamId, ReplicationMap };
+}
diff --git a/Engine/Plugins/Developer/Concert/ConcertApp/MultiUserClient/Source/MultiUserReplicationEditor/Private/MultiUserReplicationEditorModule.cpp b/Engine/Plugins/Developer/Concert/ConcertApp/MultiUserClient/Source/MultiUserReplicationEditor/Private/MultiUserReplicationEditorModule.cpp
--- a/Engine/Plugins/Developer/Concert/ConcertApp/MultiUserClient/Source/MultiUserReplicationEditor/Private/MultiUserReplicationEditorModule.cpp
+++ b/Engine/Plugins/Developer/Concert/ConcertApp/MultiUserClient/Source/MultiUserReplicationEditor/Private/MultiUserReplicationEditorModule.cpp
@@ -2,64 +2,23 @@
 
 #include "MultiUserReplicationEditorModule.h"
 
-#include "ConcertSyncSessionFlags.h"
-
-#include "ISettingsModule.h"
 #include "MultiUserReplicationEditorStyle.h"
-#include "Settings/MultiUserReplicationSettings.h"
 #include "Modules/ModuleManager.h"
 
 #define LOCTEXT_NAMESPACE "FMultiUserReplicationEditorModule"
 
 namespace UE::MultiUserReplicationEditor
 {
-	FMultiUserReplicationEditorModule::FMultiUserReplicationEditorModule()
-		: MultiUserReplicationCategory(LOCTEXT("MultiUserReplicationCategory", "Multi-User"))
-	{}
-
 	void FMultiUserReplicationEditorModule::StartupModule()
 	{
 		FMultiUserReplicationEditorStyle::Initialize();
-		RegisterSettings();
 	}
 
 	void FMultiUserReplicationEditorModule::ShutdownModule()
 	{
 		FMultiUserReplicationEditorStyle::Shutdown();
-		UnregisterSettings();
-	}
-
-	IMultiUserReplicationEditorModule::FSettingPath FMultiUserReplicationEditorModule::GetReplicationSettingsInfo() const
-	{
-		return { "Project", "Plugins", "Multi-User Replication" };
-	}
-
-	void FMultiUserReplicationEditorModule::RegisterSettings()
-	{
-		ISettingsModule* SettingsModule = FModuleManager::GetModulePtr<ISettingsModule>("Settings");
-		if (ensure(SettingsModule))
-		{
-			const FSettingPath Path = GetReplicationSettingsInfo();
-			SettingsSection = SettingsModule->RegisterSettings(Path.ContainerName, Path.CategoryName, Path.SectionName,
-				LOCTEXT("MultiUserReplicationSettingsName", "Multi-User Replication"),
-				LOCTEXT("MultiUserReplicationSettingsDescription", "Configure the Multi-User Replication settings."),
-				UMultiUserReplicationSettings::Get()
-			);
-		}
-	}
-
-	void FMultiUserReplicationEditorModule::UnregisterSettings()
-	{
-		ISettingsModule* SettingsModule = FModuleManager::GetModulePtr<ISettingsModule>("Settings");
-		if (SettingsModule)
-		{
-			const FSettingPath Path = GetReplicationSettingsInfo();
-			SettingsModule->UnregisterSettings(Path.ContainerName, Path.CategoryName, Path.SectionName);
-		}
-
-		SettingsSection = nullptr;
 	}
-};
+}
 
 IMPLEMENT_MODULE(UE::MultiUserReplicationEditor::FMultiUserReplicationEditorModule, MultiUserReplicationEditor);
 #undef LOCTEXT_NAMESPACE
\ No newline at end of file
diff --git a/Engine/Plugins/Developer/Concert/ConcertApp/MultiUserClient/Source/MultiUserReplicationEditor/Private/MultiUserReplicationEditorModule.h b/Engine/Plugins/Developer/Concert/ConcertApp/MultiUserClient/Source/MultiUserReplicationEditor/Private/MultiUserReplicationEditorModule.h
--- a/Engine/Plugins/Developer/Concert/ConcertApp/MultiUserClient/Source/MultiUserReplicationEditor/Private/MultiUserReplicationEditorModule.h
+++ b/Engine/Plugins/Developer/Concert/ConcertApp/MultiUserClient/Source/MultiUserReplicationEditor/Private/MultiUserReplicationEditorModule.h
@@ -3,33 +3,16 @@
 #pragma once
 
 #include "IMultiUserReplicationEditorModule.h"
-#include "Misc/AssetCategoryPath.h"
 
 namespace UE::MultiUserReplicationEditor
 {
 	class FMultiUserReplicationEditorModule : public IMultiUserReplicationEditorModule
 	{
 	public:
-
-		FMultiUserReplicationEditorModule();
-
+	
 		//~ Begin IModuleInterface Interface
 		virtual void StartupModule() override;
 		virtual void ShutdownModule() override;
 		//~ End IModuleInterface Interface
-
-		//~ Begin IMultiUserReplicationEditor Interface
-		virtual FAssetCategoryPath GetMultiUserReplicationCategory() const override { return MultiUserReplicationCategory; }
-		virtual FSettingPath GetReplicationSettingsInfo() const override;
-		//~ End IMultiUserReplicationEditor Interface
-
-	private:
-
-		FAssetCategoryPath MultiUserReplicationCategory;
-
-		TSharedPtr<ISettingsSection> SettingsSection;
-
-		void RegisterSettings();
-		void UnregisterSettings();
 	};
 }
diff --git a/Engine/Plugins/Developer/Concert/ConcertApp/MultiUserClient/Source/MultiUserReplicationEditor/Private/MultiUserReplicationEditorStyle.cpp b/Engine/Plugins/Developer/Concert/ConcertApp/MultiUserClient/Source/MultiUserReplicationEditor/Private/MultiUserReplicationEditorStyle.cpp
--- a/Engine/Plugins/Developer/Concert/ConcertApp/MultiUserClient/Source/MultiUserReplicationEditor/Private/MultiUserReplicationEditorStyle.cpp
+++ b/Engine/Plugins/Developer/Concert/ConcertApp/MultiUserClient/Source/MultiUserReplicationEditor/Private/MultiUserReplicationEditorStyle.cpp
@@ -45,9 +45,6 @@ namespace UE::MultiUserReplicationEditor
 		{
 			Style->SetContentRoot(FPaths::Combine(Plugin->GetBaseDir(), TEXT("Resources")));
 		}
-
-		Style->Set("ReplicationTreeView.RowHeight", 20.f);
-		Style->Set("ReplicationTreeView.IconColumnWidth", 20.f);
 		
 		return Style;
 	}
diff --git a/Engine/Plugins/Developer/Concert/ConcertUI/ConcertClientSharedSlate/Source/ConcertClientSharedSlate/ConcertClientSharedSlate.Build.cs b/Engine/Plugins/Developer/Concert/ConcertUI/ConcertClientSharedSlate/Source/ConcertClientSharedSlate/ConcertClientSharedSlate.Build.cs
--- a/Engine/Plugins/Developer/Concert/ConcertUI/ConcertClientSharedSlate/Source/ConcertClientSharedSlate/ConcertClientSharedSlate.Build.cs
+++ b/Engine/Plugins/Developer/Concert/ConcertUI/ConcertClientSharedSlate/Source/ConcertClientSharedSlate/ConcertClientSharedSlate.Build.cs
@@ -9,15 +9,41 @@ namespace UnrealBuildTool.Rules
 			PublicDependencyModuleNames.AddRange(
 				new string[]
 				{
+					// Engine
 					"Core",
 					"CoreUObject",
+					"Engine",
 					"Slate",
 					"SlateCore",
 					"EditorStyle",
+					
+					// Concert
 					"ConcertSharedSlate",
 					"ConcertSyncClient",
 				}
 			);
+			
+			PrivateDependencyModuleNames.AddRange(
+				new string[]
+				{
+					// Engine
+					"AssetDefinition",
+					"AssetRegistry",
+					"InputCore",
+					"Projects",
+					"Slate",
+					"SlateCore",
+					"ToolMenus",
+					"ToolWidgets",
+					"UnrealEd", 
+					
+					// Concert
+					"Concert",
+					"ConcertClient",
+					"ConcertSyncCore",
+					"ConcertTransport", // For LogConcert
+				}
+			);
 
 			ShortName = "ConClShrSlt";
 		}
diff --git a/Engine/Plugins/Developer/Concert/ConcertUI/ConcertClientSharedSlate/Source/ConcertClientSharedSlate/Private/Replication/Editor/Model/GenericPropertySelectionModel.cpp b/Engine/Plugins/Developer/Concert/ConcertUI/ConcertClientSharedSlate/Source/ConcertClientSharedSlate/Private/Replication/Editor/Model/GenericPropertySelectionModel.cpp
new file mode 100644
--- /dev/null
+++ b/Engine/Plugins/Developer/Concert/ConcertUI/ConcertClientSharedSlate/Source/ConcertClientSharedSlate/Private/Replication/Editor/Model/GenericPropertySelectionModel.cpp
@@ -0,0 +1,304 @@
+// Copyright Epic Games, Inc. All Rights Reserved.
+
+#include "GenericPropertySelectionModel.h"
+
+#include "Replication/PropertyChainUtils.h"
+#include "Replication/Data/ObjectReplicationMap.h"
+#include "Replication/Settings/ConcertReplicationEditorSettings.h"
+
+#include "Algo/AllOf.h"
+
+namespace UE::ConcertClientSharedSlate
+{
+	FGenericPropertySelectionModel::FGenericPropertySelectionModel(
+		TAttribute<FObjectReplicationMap*> ReplicationMapAttribute,
+		TAttribute<const FConcertReplicationEditorSettings*> OptionalReplicationSettingsAttribute
+		)
+		: ReplicationMapAttribute(MoveTemp(ReplicationMapAttribute))
+		, OptionalReplicationSettingsAttribute(MoveTemp(OptionalReplicationSettingsAttribute))
+	{}
+	
+	uint32 FGenericPropertySelectionModel::GetNumReplicatedObjects() const
+	{
+		const FObjectReplicationMap* ReplicationMap = ReplicationMapAttribute.Get();
+		return ensure(ReplicationMap)
+			? ReplicationMap->ReplicatedObjects.Num()
+			: 0;
+	}
+
+	uint32 FGenericPropertySelectionModel::GetNumProperties(const FSoftObjectPath& Object) const
+	{
+		const FObjectReplicationMap* ReplicationMap = ReplicationMapAttribute.Get();
+		if (!ensure(ReplicationMap))
+		{
+			return false;
+		}
+
+		const FReplicatedObjectInfo* AssignedProperties = ReplicationMap->ReplicatedObjects.Find(Object);
+		return AssignedProperties
+			? AssignedProperties->PropertySelection.ReplicatedProperties.Num()
+			: 0;
+	}
+
+	FSoftClassPath FGenericPropertySelectionModel::GetObjectClass(const FSoftObjectPath& Object) const
+	{
+		const FObjectReplicationMap* ReplicationMap = ReplicationMapAttribute.Get();
+		if (!ensure(ReplicationMap))
+		{
+			return {};
+		}
+		
+		const FReplicatedObjectInfo* AssignedProperties = ReplicationMap->ReplicatedObjects.Find(Object);
+		return AssignedProperties
+			? AssignedProperties->ClassPath
+			: FSoftClassPath{};
+	}
+
+	bool FGenericPropertySelectionModel::ContainsObjects(const TSet<FSoftObjectPath>& Objects) const
+	{
+		const FObjectReplicationMap* ReplicationMap = ReplicationMapAttribute.Get();
+		return ensure(ReplicationMap)
+			&& Algo::AllOf(Objects, [this, ReplicationMap](const FSoftObjectPath& ObjectPath){ return ReplicationMap->ReplicatedObjects.Contains(ObjectPath); });
+	}
+
+	bool FGenericPropertySelectionModel::ContainsProperties(const FSoftObjectPath& Object, const TSet<FConcertPropertyChain>& Properties) const
+	{
+		const FObjectReplicationMap* ReplicationMap = ReplicationMapAttribute.Get();
+		if (!ensure(ReplicationMap))
+		{
+			return false;
+		}
+
+		const FReplicatedObjectInfo* ObjectInfo = ReplicationMap->ReplicatedObjects.Find(Object);
+		return ObjectInfo
+			&& Algo::AllOf(Properties, [ObjectInfo](const FConcertPropertyChain& Property){ return ObjectInfo->PropertySelection.ReplicatedProperties.Contains(Property); });
+	}
+
+	bool FGenericPropertySelectionModel::ForEachReplicatedObject(TFunctionRef<EBreakBehavior(const FSoftObjectPath& Object)> Delegate) const
+	{
+		const FObjectReplicationMap* ReplicationMap = ReplicationMapAttribute.Get();
+		if (!ensure(ReplicationMap))
+		{
+			return false;
+		}
+
+		for (const TPair<FSoftObjectPath, FReplicatedObjectInfo>& ObjectMap: ReplicationMap->ReplicatedObjects)
+		{
+			if (Delegate(ObjectMap.Key) == EBreakBehavior::Break)
+			{
+				return true;
+			}
+		}
+
+		return !ReplicationMap->ReplicatedObjects.IsEmpty();
+	}
+
+	bool FGenericPropertySelectionModel::ForEachProperty(const FSoftObjectPath& Object, TFunctionRef<EBreakBehavior(const FConcertPropertyChain& Parent)> Delegate) const
+	{
+		const FObjectReplicationMap* ReplicationMap = ReplicationMapAttribute.Get();
+		if (!ensure(ReplicationMap))
+		{
+			return false;
+		}
+
+		const FReplicatedObjectInfo* AssignedProperties = ReplicationMap->ReplicatedObjects.Find(Object);
+		if (!AssignedProperties)
+		{
+			return false;
+		}
+
+		for (const FConcertPropertyChain& ReplicatedPropertyInfo : AssignedProperties->PropertySelection.ReplicatedProperties)
+		{
+			if (Delegate(ReplicatedPropertyInfo) == EBreakBehavior::Break)
+			{
+				return true;
+			}
+		}
+		return !AssignedProperties->PropertySelection.ReplicatedProperties.IsEmpty();
+	}
+
+	void FGenericPropertySelectionModel::AddObjects(TArrayView<UObject*> Objects)
+	{
+		FObjectReplicationMap* ReplicationMap = ReplicationMapAttribute.Get();
+		if (!ensure(ReplicationMap) || Objects.IsEmpty())
+		{
+			return;
+		}
+		
+		TSet<UObject*> ObjectsNotAdded;
+		TArray<UObject*> AdditionalObjectsToAdd;
+		for (UObject* Object : Objects)
+		{
+			const FSoftObjectPath ObjectPath = Object;
+			if (ensureAlways(Object) && !ReplicationMap->ReplicatedObjects.Contains(ObjectPath))
+			{
+				FReplicatedObjectInfo& ObjectInfo = ReplicationMap->ReplicatedObjects.Add(ObjectPath);
+				ObjectInfo.ClassPath = Object->GetClass();
+
+				const FConcertReplicationEditorSettings* AutoPopulateSettings = OptionalReplicationSettingsAttribute.IsBound()
+					? OptionalReplicationSettingsAttribute.Get()
+					: nullptr;
+				if (!AutoPopulateSettings)
+				{
+					continue;
+				}
+				
+				AutoPopulateSettings->AddDefaultPropertiesFromSettings(ObjectInfo, *Object->GetClass());
+				AutoPopulateSettings->AddAdditionalObjectsFromSettings(*Object, [&AdditionalObjectsToAdd](UObject& FurtherObject)
+				{
+					AdditionalObjectsToAdd.Add(&FurtherObject);
+				});
+			}
+			else
+			{
+				ObjectsNotAdded.Add(Object);
+			}
+		}
+		
+		if (LIKELY(ObjectsNotAdded.IsEmpty()))
+		{
+			OnObjectsChangedDelegate.Broadcast(Objects, {}, EReplicatedObjectChangeReason::ChangedDirectly);
+		}
+		else if (ObjectsNotAdded.Num() < Objects.Num())
+		{
+			// Uncommon case so it is ok if it is suboptimal
+			TArray<UObject*> Added;
+			for (UObject* Object : Objects)
+			{
+				if (!ObjectsNotAdded.Contains(Object))
+				{
+					Added.Add(Object);
+				}
+			}
+			OnObjectsChangedDelegate.Broadcast(Added, {}, EReplicatedObjectChangeReason::ChangedDirectly);
+		}
+
+		AddObjects(AdditionalObjectsToAdd);
+	}
+
+	void FGenericPropertySelectionModel::RemoveObjects(TArrayView<FSoftObjectPath> Objects)
+	{
+		FObjectReplicationMap* ReplicationMap = ReplicationMapAttribute.Get();
+		if (!ensure(ReplicationMap) || Objects.IsEmpty())
+		{
+			return;
+		}
+		
+		TSet<FSoftObjectPath> ObjectsNotRemoved;
+		for (const FSoftObjectPath& Object : Objects)
+		{
+			const bool bRemoved = ReplicationMap->ReplicatedObjects.Remove(Object) != 0;
+			if (!bRemoved)
+			{
+				ObjectsNotRemoved.Add(Object);
+			}
+		}
+		
+		if (LIKELY(ObjectsNotRemoved.IsEmpty()))
+		{
+			OnObjectsChangedDelegate.Broadcast({}, Objects, EReplicatedObjectChangeReason::ChangedDirectly);
+		}
+		else if (ObjectsNotRemoved.Num() < Objects.Num())
+		{
+			// Uncommon case so it is ok if it is suboptimal
+			TArray<FSoftObjectPath> Removed;
+			for (const FSoftObjectPath& Object : Objects)
+			{
+				if (!ObjectsNotRemoved.Contains(Object))
+				{
+					Removed.Add(Object);
+				}
+			}
+			OnObjectsChangedDelegate.Broadcast({}, Removed, EReplicatedObjectChangeReason::ChangedDirectly);
+		}
+	}
+
+	void FGenericPropertySelectionModel::AddProperties(const FSoftObjectPath& Object, TArrayView<FConcertPropertyChain> Properties)
+	{
+		FObjectReplicationMap* ReplicationMap = ReplicationMapAttribute.Get();
+		if (!ensure(ReplicationMap))
+		{
+			return;
+		}
+		
+		FReplicatedObjectInfo* AssignedProperties = ReplicationMap->ReplicatedObjects.Find(Object);
+		if (!AssignedProperties)
+		{
+			return;
+		}
+
+		UClass* Class = AssignedProperties->ClassPath.TryLoadClass<UObject>();
+		if (!Class)
+		{
+			return;
+		}
+
+		bool bAddedAtLeastOne = false;
+		TArray<FConcertPropertyChain>& ReplicatedProperties = AssignedProperties->PropertySelection.ReplicatedProperties;
+		ReplicatedProperties.Reserve(ReplicatedProperties.Num() + Properties.Num());
+		for (const FConcertPropertyChain& AddedProperty : Properties)
+		{
+			bAddedAtLeastOne |= ReplicatedProperties.AddUnique(AddedProperty) != INDEX_NONE;
+
+			// Parent properties must also be added
+			// Not exactly efficient to iterate through the hierarchy for every removed item but it should be fine... Properties.Num() == 1 is the most common case
+			ConcertSyncCore::PropertyChain::ForEachReplicatableConcertProperty(*Class, [&bAddedAtLeastOne, &AssignedProperties, &AddedProperty](FConcertPropertyChain&& Property)
+			{
+				if (Property.IsParentOf(AddedProperty))
+				{
+					bAddedAtLeastOne |= AssignedProperties->PropertySelection.ReplicatedProperties.AddUnique(Property) != INDEX_NONE;
+				}
+				return EBreakBehavior::Continue;
+			});
+		}
+
+		if (bAddedAtLeastOne)
+		{
+			OnPropertiesChangedDelegate.Broadcast();
+		}
+	}
+
+	void FGenericPropertySelectionModel::RemoveProperties(const FSoftObjectPath& Object, TArrayView<FConcertPropertyChain> Properties)
+	{
+		FObjectReplicationMap* ReplicationMap = ReplicationMapAttribute.Get();
+		if (!ensure(ReplicationMap))
+		{
+			return;
+		}
+		
+		FReplicatedObjectInfo* AssignedProperties = ReplicationMap->ReplicatedObjects.Find(Object);
+		if (!AssignedProperties)
+		{
+			return;
+		}
+
+		UClass* Class = AssignedProperties->ClassPath.TryLoadClass<UObject>();
+		int32 NumRemoved = 0;
+		for (const FConcertPropertyChain& RemovedProperty : Properties)
+		{
+			NumRemoved += AssignedProperties->PropertySelection.ReplicatedProperties.Remove(RemovedProperty);
+
+			// Removal should not fail if the class is not available
+			if (!Class)
+			{
+				continue;
+			}
+
+			// Child properties must also be removed
+			// Not exactly efficient to iterate through the hierarchy for every removed item but it should be fine... Properties.Num() == 1 is the most common case
+			ConcertSyncCore::PropertyChain::ForEachReplicatableConcertProperty(*Class, [&NumRemoved, &AssignedProperties, &RemovedProperty](FConcertPropertyChain&& Property)
+			{
+				if (Property.IsChildOf(RemovedProperty))
+				{
+					NumRemoved += AssignedProperties->PropertySelection.ReplicatedProperties.Remove(Property);
+				}
+				return EBreakBehavior::Continue;
+			});
+		}
+		
+		if (NumRemoved > 0)
+		{
+			OnPropertiesChangedDelegate.Broadcast();
+		}
+	}
+}
diff --git a/Engine/Plugins/Developer/Concert/ConcertUI/ConcertClientSharedSlate/Source/ConcertClientSharedSlate/Private/Replication/Editor/Model/GenericPropertySelectionModel.h b/Engine/Plugins/Developer/Concert/ConcertUI/ConcertClientSharedSlate/Source/ConcertClientSharedSlate/Private/Replication/Editor/Model/GenericPropertySelectionModel.h
new file mode 100644
--- /dev/null
+++ b/Engine/Plugins/Developer/Concert/ConcertUI/ConcertClientSharedSlate/Source/ConcertClientSharedSlate/Private/Replication/Editor/Model/GenericPropertySelectionModel.h
@@ -0,0 +1,58 @@
+// Copyright Epic Games, Inc. All Rights Reserved.
+
+#pragma once
+
+#include "Replication/Editor/Model/IEditableObjectToPropertiesModel.h"
+#include "Misc/Attribute.h"
+#include "UObject/WeakObjectPtrTemplates.h"
+
+struct FConcertReplicationEditorSettings;
+class FTransactionObjectEvent;
+struct FObjectReplicationMap;
+struct FTransactionContext;
+
+namespace UE::ConcertClientSharedSlate
+{
+	/** Implements logic for editing a FObjectReplicationMap contained in an UObject. */
+	class FGenericPropertySelectionModel
+		: public IEditableObjectToPropertiesModel
+	{
+	public:
+		
+		FGenericPropertySelectionModel(
+			TAttribute<FObjectReplicationMap*> ReplicationMapAttribute,
+			TAttribute<const FConcertReplicationEditorSettings*> OptionalReplicationSettingsAttribute = {}
+			);
+		
+		//~ Begin IObjectToPropertiesModel Interface
+		virtual uint32 GetNumReplicatedObjects() const override;
+		virtual uint32 GetNumProperties(const FSoftObjectPath& Object) const override;
+		virtual FSoftClassPath GetObjectClass(const FSoftObjectPath& Object) const override;
+		virtual bool ContainsObjects(const TSet<FSoftObjectPath>& Objects) const override;
+		virtual bool ContainsProperties(const FSoftObjectPath& Object, const TSet<FConcertPropertyChain>& Properties) const override;
+		virtual bool ForEachReplicatedObject(TFunctionRef<EBreakBehavior(const FSoftObjectPath& Object)> Delegate) const override;
+		virtual bool ForEachProperty(const FSoftObjectPath& Object, TFunctionRef<EBreakBehavior(const FConcertPropertyChain& Parent)> Delegate) const override;
+		//~ End IObjectToPropertiesModel Interface
+		
+		//~ Begin IEditableObjectToPropertiesModel Interface
+		virtual void AddObjects(TArrayView<UObject*> Objects) override;
+		virtual void RemoveObjects(TArrayView<FSoftObjectPath> Objects) override;
+		virtual void AddProperties(const FSoftObjectPath&, TArrayView<FConcertPropertyChain> Properties) override;
+		virtual void RemoveProperties(const FSoftObjectPath&, TArrayView<FConcertPropertyChain> Properties) override;
+		virtual FOnObjectsChanged& OnObjectsChanged() override { return OnObjectsChangedDelegate; }
+		virtual FOnPropertiesChanged& OnPropertiesChanged() override { return OnPropertiesChangedDelegate; }
+		//~ End IEditableObjectToPropertiesModel Interface
+
+	private:
+
+		/** Returns the replication map that is supposed to be edited. */
+		TAttribute<FObjectReplicationMap*> ReplicationMapAttribute;
+
+		/** Optional. Contains settings for auto populating objects. */
+		TAttribute<const FConcertReplicationEditorSettings*> OptionalReplicationSettingsAttribute;
+
+		FOnObjectsChanged OnObjectsChangedDelegate;
+		FOnPropertiesChanged OnPropertiesChangedDelegate;
+	};
+}
+
diff --git a/Engine/Plugins/Developer/Concert/ConcertUI/ConcertClientSharedSlate/Source/ConcertClientSharedSlate/Private/Replication/Editor/Model/Property/SelectPropertyFromUClassModel.cpp b/Engine/Plugins/Developer/Concert/ConcertUI/ConcertClientSharedSlate/Source/ConcertClientSharedSlate/Private/Replication/Editor/Model/Property/SelectPropertyFromUClassModel.cpp
new file mode 100644
--- /dev/null
+++ b/Engine/Plugins/Developer/Concert/ConcertUI/ConcertClientSharedSlate/Source/ConcertClientSharedSlate/Private/Replication/Editor/Model/Property/SelectPropertyFromUClassModel.cpp
@@ -0,0 +1,23 @@
+// Copyright Epic Games, Inc. All Rights Reserved.
+
+#include "Replication/Editor/Model/Property/SelectPropertyFromUClassModel.h"
+
+#include "Replication/Editor/Model/Property/ConcertSyncCoreReplicatedPropertySource.h"
+
+#define LOCTEXT_NAMESPACE "FSelectPropertyFromUClassModel"
+
+namespace UE::ConcertClientSharedSlate
+{
+	FSelectPropertyFromUClassModel::FSelectPropertyFromUClassModel()
+		: UClassIteratorSource(MakeShared<FConcertSyncCoreReplicatedPropertySource>())
+	{}
+
+	TSharedRef<IPropertySourceModel> FSelectPropertyFromUClassModel::GetPropertySource(const FSoftClassPath& Class) const
+	{
+		UClass* LoadedClass = Class.TryLoadClass<UObject>();
+		UClassIteratorSource->SetClass(LoadedClass);
+		return UClassIteratorSource;
+	}
+}
+
+#undef LOCTEXT_NAMESPACE
\ No newline at end of file
diff --git a/Engine/Plugins/Developer/Concert/ConcertUI/ConcertClientSharedSlate/Source/ConcertClientSharedSlate/Private/Replication/Editor/Model/TransactionalPropertySelectionModel.cpp b/Engine/Plugins/Developer/Concert/ConcertUI/ConcertClientSharedSlate/Source/ConcertClientSharedSlate/Private/Replication/Editor/Model/TransactionalPropertySelectionModel.cpp
new file mode 100644
--- /dev/null
+++ b/Engine/Plugins/Developer/Concert/ConcertUI/ConcertClientSharedSlate/Source/ConcertClientSharedSlate/Private/Replication/Editor/Model/TransactionalPropertySelectionModel.cpp
@@ -0,0 +1,84 @@
+// Copyright Epic Games, Inc. All Rights Reserved.
+
+#include "TransactionalPropertySelectionModel.h"
+#include "ScopedTransaction.h"
+#include "Misc/TransactionObjectEvent.h"
+
+#define LOCTEXT_NAMESPACE "FTransactionalPropertySelectionModel"
+
+namespace UE::ConcertClientSharedSlate
+{
+	FTransactionalPropertySelectionModel::FTransactionalPropertySelectionModel(
+		UObject& OwningObject,
+		TAttribute<FObjectReplicationMap*> ReplicationMapAttribute,
+		TAttribute<const FConcertReplicationEditorSettings*> OptionalReplicationSettingsAttribute
+		)
+		: FGenericPropertySelectionModel(MoveTemp(ReplicationMapAttribute), MoveTemp(OptionalReplicationSettingsAttribute))
+		, OwningObject(&OwningObject)
+	{}
+
+	void FTransactionalPropertySelectionModel::AddObjects(TArrayView<UObject*> Objects)
+	{
+		const FScopedTransaction Transaction(LOCTEXT("AddObjects", "Add replicated objects"));
+		OwningObject->Modify();
+		FGenericPropertySelectionModel::AddObjects(Objects);
+	}
+
+	void FTransactionalPropertySelectionModel::RemoveObjects(TArrayView<FSoftObjectPath> Objects)
+	{
+		const FScopedTransaction Transaction(LOCTEXT("RemoveObjects", "Remove replicated objects"));
+		OwningObject->Modify();
+		FGenericPropertySelectionModel::RemoveObjects(Objects);
+	}
+
+	void FTransactionalPropertySelectionModel::AddProperties(const FSoftObjectPath& SoftObjectPath, TArrayView<FConcertPropertyChain> Properties)
+	{
+		const FScopedTransaction Transaction(LOCTEXT("AddProperties", "Add replicated properties"));
+		OwningObject->Modify();
+		FGenericPropertySelectionModel::AddProperties(SoftObjectPath, Properties);
+	}
+
+	void FTransactionalPropertySelectionModel::RemoveProperties(const FSoftObjectPath& SoftObjectPath, TArrayView<FConcertPropertyChain> Properties)
+	{
+		const FScopedTransaction Transaction(LOCTEXT("RemoveProperties", "Remove replicated properties"));
+		OwningObject->Modify();
+		FGenericPropertySelectionModel::RemoveProperties(SoftObjectPath, Properties);
+	}
+
+	bool FTransactionalPropertySelectionModel::MatchesContext(const FTransactionContext& InContext, const TArray<TPair<UObject*, FTransactionObjectEvent>>& TransactionObjectContexts) const
+	{
+		if (!ensure(OwningObject.IsValid()))
+		{
+			return false;
+		}
+		
+		for (const TPair<UObject*, FTransactionObjectEvent>& TransactionObjectPair : TransactionObjectContexts)
+		{
+			UObject* Object = TransactionObjectPair.Key;
+			while (Object != nullptr)
+			{
+				if (Object == OwningObject)
+				{
+					return true;
+				}
+				Object = Object->GetOuter();
+			}
+		}
+
+		return false;
+	}
+
+	void FTransactionalPropertySelectionModel::PostUndo(bool bSuccess)
+	{
+		OnObjectsChanged().Broadcast({}, {}, EReplicatedObjectChangeReason::Transacted);
+		OnPropertiesChanged().Broadcast();
+	}
+
+	void FTransactionalPropertySelectionModel::PostRedo(bool bSuccess)
+	{
+		OnObjectsChanged().Broadcast({}, {}, EReplicatedObjectChangeReason::Transacted);
+		OnPropertiesChanged().Broadcast();
+	}
+}
+
+#undef LOCTEXT_NAMESPACE
\ No newline at end of file
diff --git a/Engine/Plugins/Developer/Concert/ConcertUI/ConcertClientSharedSlate/Source/ConcertClientSharedSlate/Private/Replication/Editor/Model/TransactionalPropertySelectionModel.h b/Engine/Plugins/Developer/Concert/ConcertUI/ConcertClientSharedSlate/Source/ConcertClientSharedSlate/Private/Replication/Editor/Model/TransactionalPropertySelectionModel.h
new file mode 100644
--- /dev/null
+++ b/Engine/Plugins/Developer/Concert/ConcertUI/ConcertClientSharedSlate/Source/ConcertClientSharedSlate/Private/Replication/Editor/Model/TransactionalPropertySelectionModel.h
@@ -0,0 +1,43 @@
+// Copyright Epic Games, Inc. All Rights Reserved.
+
+#pragma once
+
+#include "GenericPropertySelectionModel.h"
+#include "EditorUndoClient.h"
+
+namespace UE::ConcertClientSharedSlate
+{
+	/** Special case of FGenericPropertySelectionModel where the edited FObjectReplicationMap lives in an UObject that is RF_Transactional. */
+	class FTransactionalPropertySelectionModel
+		: public FGenericPropertySelectionModel
+		, public FSelfRegisteringEditorUndoClient
+	{
+	public:
+
+		FTransactionalPropertySelectionModel(
+			UObject& OwningObject,
+			TAttribute<FObjectReplicationMap*> ReplicationMapAttribute,
+			TAttribute<const FConcertReplicationEditorSettings*> OptionalReplicationSettingsAttribute = {}
+			);
+		
+		//~ Begin IEditableObjectToPropertiesModel Interface
+		virtual void AddObjects(TArrayView<UObject*> Objects) override;
+		virtual void RemoveObjects(TArrayView<FSoftObjectPath> Objects) override;
+		virtual void AddProperties(const FSoftObjectPath&, TArrayView<FConcertPropertyChain> Properties) override;
+		virtual void RemoveProperties(const FSoftObjectPath&, TArrayView<FConcertPropertyChain> Properties) override;
+		//~ End IEditableObjectToPropertiesModel Interface
+
+		//~ Begin FEditorUndoClient Interface
+		virtual bool MatchesContext(const FTransactionContext& InContext, const TArray<TPair<UObject*, FTransactionObjectEvent>>& TransactionObjectContexts) const override;
+		virtual void PostUndo(bool bSuccess) override;
+		virtual void PostRedo(bool bSuccess) override;
+		//~ End FEditorUndoClient Interface
+
+	private:
+
+		/** User of FTransactionalPropertySelectionModel is responsible for keeping OwningObject alive, e.g. via an asset editor. */
+		TWeakObjectPtr<UObject> OwningObject;
+	};
+}
+
+
diff --git a/Engine/Plugins/Developer/Concert/ConcertUI/ConcertClientSharedSlate/Source/ConcertClientSharedSlate/Private/Replication/Editor/View/ObjectEditor/FakeObjectToPropertiesEditorModel.cpp b/Engine/Plugins/Developer/Concert/ConcertUI/ConcertClientSharedSlate/Source/ConcertClientSharedSlate/Private/Replication/Editor/View/ObjectEditor/FakeObjectToPropertiesEditorModel.cpp
new file mode 100644
--- /dev/null
+++ b/Engine/Plugins/Developer/Concert/ConcertUI/ConcertClientSharedSlate/Source/ConcertClientSharedSlate/Private/Replication/Editor/View/ObjectEditor/FakeObjectToPropertiesEditorModel.cpp
@@ -0,0 +1,20 @@
+// Copyright Epic Games, Inc. All Rights Reserved.
+
+#include "FakeObjectToPropertiesEditorModel.h"
+
+#include "Replication/Editor/Model/Property/IPropertySelectionSourceModel.h"
+#include "Replication/Editor/Model/Property/IPropertySourceModel.h"
+#include "Replication/PropertyChainUtils.h"
+
+namespace UE::ConcertClientSharedSlate
+{
+	bool FFakeObjectToPropertiesEditorModel::ForEachProperty(const FSoftObjectPath& Object, TFunctionRef<EBreakBehavior(const FConcertPropertyChain& Property)> Delegate) const
+	{
+		PropertySelectionSource->GetPropertySource(GetObjectClass(Object))
+			->EnumerateSelectableItems([&Delegate](const FSelectablePropertyInfo& PropertyInfo)
+			{
+				return Delegate(PropertyInfo.Property);
+			});
+		return true;
+	}
+}
diff --git a/Engine/Plugins/Developer/Concert/ConcertUI/ConcertClientSharedSlate/Source/ConcertClientSharedSlate/Private/Replication/Editor/View/SReplicationStreamEditor.cpp b/Engine/Plugins/Developer/Concert/ConcertUI/ConcertClientSharedSlate/Source/ConcertClientSharedSlate/Private/Replication/Editor/View/SReplicationStreamEditor.cpp
new file mode 100644
--- /dev/null
+++ b/Engine/Plugins/Developer/Concert/ConcertUI/ConcertClientSharedSlate/Source/ConcertClientSharedSlate/Private/Replication/Editor/View/SReplicationStreamEditor.cpp
@@ -0,0 +1,26 @@
+// Copyright Epic Games, Inc. All Rights Reserved.
+
+#include "SReplicationStreamEditor.h"
+
+#include "Replication/ReplicationWidgetFactories.h"
+#include "Replication/Editor/View/ObjectEditor/SPropertyReplicationSelectionEditor.h"
+
+#include "Widgets/SBoxPanel.h"
+
+namespace UE::ConcertClientSharedSlate
+{
+	void SReplicationStreamEditor::Construct(const FArguments& InArgs, const FCreateEditorParams& Params)
+	{
+		ChildSlot
+		[
+			SNew(SVerticalBox)
+
+			+SVerticalBox::Slot()
+			[
+				SNew(SPropertyReplicationSelectionEditor, Params.DataModel, Params.ObjectSource, Params.PropertySource)
+			]
+		];
+	}
+}
+
+
diff --git a/Engine/Plugins/Developer/Concert/ConcertUI/ConcertClientSharedSlate/Source/ConcertClientSharedSlate/Private/Replication/Editor/View/SReplicationStreamEditor.h b/Engine/Plugins/Developer/Concert/ConcertUI/ConcertClientSharedSlate/Source/ConcertClientSharedSlate/Private/Replication/Editor/View/SReplicationStreamEditor.h
new file mode 100644
--- /dev/null
+++ b/Engine/Plugins/Developer/Concert/ConcertUI/ConcertClientSharedSlate/Source/ConcertClientSharedSlate/Private/Replication/Editor/View/SReplicationStreamEditor.h
@@ -0,0 +1,26 @@
+// Copyright Epic Games, Inc. All Rights Reserved.
+
+#pragma once
+
+#include "Widgets/DeclarativeSyntaxSupport.h"
+#include "Widgets/SCompoundWidget.h"
+
+class UMultiUserReplicationClientPreset;
+
+namespace UE::ConcertClientSharedSlate
+{
+	class SPropertyReplicationSelectionEditor;
+	
+	struct FCreateEditorParams;
+	
+	/** Root widget for editing UMultiUserReplicationStreamAsset. */
+	class CONCERTCLIENTSHAREDSLATE_API SReplicationStreamEditor : public SCompoundWidget
+	{
+	public:
+		SLATE_BEGIN_ARGS(SReplicationStreamEditor)
+		{}
+		SLATE_END_ARGS()
+
+		void Construct(const FArguments& InArgs, const FCreateEditorParams& Params);
+	};
+}
diff --git a/Engine/Plugins/Developer/Concert/ConcertUI/ConcertClientSharedSlate/Source/ConcertClientSharedSlate/Private/Replication/ReplicationWidgetFactories.cpp b/Engine/Plugins/Developer/Concert/ConcertUI/ConcertClientSharedSlate/Source/ConcertClientSharedSlate/Private/Replication/ReplicationWidgetFactories.cpp
new file mode 100644
--- /dev/null
+++ b/Engine/Plugins/Developer/Concert/ConcertUI/ConcertClientSharedSlate/Source/ConcertClientSharedSlate/Private/Replication/ReplicationWidgetFactories.cpp
@@ -0,0 +1,28 @@
+// Copyright Epic Games, Inc. All Rights Reserved.
+
+#include "Replication/ReplicationWidgetFactories.h"
+
+#include "Editor/Model/TransactionalPropertySelectionModel.h"
+#include "Editor/View/SReplicationStreamEditor.h"
+
+namespace UE::ConcertClientSharedSlate
+{
+	TSharedRef<SWidget> CreateEditor(FCreateEditorParams Params)
+	{
+		return SNew(SReplicationStreamEditor, MoveTemp(Params));
+	}
+	
+	TSharedRef<IEditableObjectToPropertiesModel> CreatePropertySelectionModel(
+		UObject& OwnerObject,
+		TAttribute<FObjectReplicationMap*> ReplicationMapAttribute,
+		TAttribute<const FConcertReplicationEditorSettings*> OptionalReplicationSettingsAttribute
+		)
+	{
+		return MakeShared<FTransactionalPropertySelectionModel>(
+			OwnerObject,
+			MoveTemp(ReplicationMapAttribute),
+			MoveTemp(OptionalReplicationSettingsAttribute)
+			);
+	}
+}
+
diff --git a/Engine/Plugins/Developer/Concert/ConcertUI/ConcertClientSharedSlate/Source/ConcertClientSharedSlate/Private/Replication/Settings/ConcertReplicationEditorSettings.cpp b/Engine/Plugins/Developer/Concert/ConcertUI/ConcertClientSharedSlate/Source/ConcertClientSharedSlate/Private/Replication/Settings/ConcertReplicationEditorSettings.cpp
new file mode 100644
--- /dev/null
+++ b/Engine/Plugins/Developer/Concert/ConcertUI/ConcertClientSharedSlate/Source/ConcertClientSharedSlate/Private/Replication/Settings/ConcertReplicationEditorSettings.cpp
@@ -0,0 +1,167 @@
+// Copyright Epic Games, Inc. All Rights Reserved.
+
+#include "Replication/Settings/ConcertReplicationEditorSettings.h"
+
+#include "ConcertLogGlobal.h"
+#include "Replication/Data/ObjectReplicationMap.h"
+#include "Replication/PropertyChainUtils.h"
+#include "Replication/Settings/ConcertDefaultPropertySelection.h"
+#include "Replication/Settings/ConcertDefaultSubobjectSelection.h"
+
+#include "Algo/IndexOf.h"
+#include "Components/ActorComponent.h"
+#include "Internationalization/Regex.h"
+
+namespace UE::ConcertClientSharedSlate::DefaultProperties
+{
+	static void ApplyDefaultPropertySelection(FReplicatedObjectInfo& Info, const FConcertDefaultPropertySelection& Selection, UStruct& Class)
+	{
+		// Preparse the the FStrings into paths
+		TArray<TArray<FName>> Paths;
+		Paths.Reserve(Selection.DefaultSelectedProperties.Num());
+		for (const FString& StringPath : Selection.DefaultSelectedProperties)
+		{
+			TArray<FString> Parts;
+			StringPath.ParseIntoArray(Parts, TEXT("."));
+			TArray<FName> NameArray;
+			Algo::Transform(Parts, NameArray, [](const FString& String){ return *String; });
+			Paths.Emplace(MoveTemp(NameArray));
+		}
+
+		// This walks through the entire property hierarchy once
+		using namespace ConcertSyncCore::PropertyChain;
+		BulkConstructConcertChainsFromPaths(Class, Paths.Num(), [&Info, &Paths](const FArchiveSerializedPropertyChain& Chain, const FProperty& LeafProperty)
+		{
+			const int32 IndexOfMatches = Algo::IndexOfByPredicate(Paths, [&Chain, &LeafProperty](const TArray<FName>& Path){ return DoPathAndChainsMatch(Path, Chain, LeafProperty); });
+			const bool bMatches = IndexOfMatches != INDEX_NONE;
+			if (bMatches)
+			{
+				Info.PropertySelection.ReplicatedProperties.Emplace(&Chain, LeafProperty);
+				Paths.RemoveAtSwap(IndexOfMatches);
+			}
+			return true;
+		});
+
+		// Warn the user of unmatched properties so they can de-clutter their settings of useless entries
+		for (const TArray<FName>& Path : Paths)
+		{
+			UE_LOG(LogConcert, Warning,
+				TEXT("Path %s was unmatched. Fix your settings for class %s."),
+				*FString::JoinBy(Path, TEXT("."), [](const FName& Name){ return Name.ToString(); }),
+				*Class.GetName()
+				);
+		}
+	}
+}
+
+void FConcertReplicationEditorSettings::AddDefaultPropertiesFromSettings(FReplicatedObjectInfo& Info, UClass& Class) const
+{
+	// Find the most specialized class properties
+	UClass* Current = &Class;
+	const FConcertDefaultPropertySelection* DefaultProperties = nullptr;
+	for (; Current && !DefaultProperties; Current = Current->GetSuperClass())
+	{
+		DefaultProperties = DefaultPropertySelection.Find(Current);
+		if (!DefaultProperties)
+		{
+			continue;
+		}
+			
+		UE::ConcertClientSharedSlate::DefaultProperties::ApplyDefaultPropertySelection(Info, *DefaultProperties, *Current);
+		// Recurse super structs
+		if (UClass* Parent = Current->GetSuperClass()
+			; Parent && DefaultProperties->bInheritFromBase)
+		{
+			AddDefaultPropertiesFromSettings(Info, *Parent);
+		}
+	}
+}
+
+namespace UE::ConcertClientSharedSlate::DefaultSubobjects
+{
+	static bool MatchesAnyRegex(const FString& Input, const TSet<FString>& AllRegex)
+	{
+		for (const FString& RegexString : AllRegex)
+		{
+			const FRegexPattern Pattern(RegexString);
+			FRegexMatcher Matcher(Pattern, Input);
+			if (Matcher.FindNext())
+			{
+				return true;
+			}
+		}
+
+		return false;
+	}
+	
+	static void ApplyDefaultSubobjectSelection(UObject& AddedObject, const FConcertDefaultSubobjectSelection& Selection, TFunctionRef<void(UObject&)> FurtherObjectsCallback)
+	{
+		// Do not search recursively so FurtherObjectsCallback can decide to call AddAdditionalObjectsFromSettings again on the newly added objects.
+		constexpr bool bIncludeNested = false;
+		ForEachObjectWithOuter(&AddedObject, [&Selection, &FurtherObjectsCallback](UObject* Subobject)
+		{
+			// Has exclusion regex?
+			const bool bShouldExclude = MatchesAnyRegex(Subobject->GetName(), Selection.ExcludeSubobjectRegex);
+			if (bShouldExclude)
+			{
+				return;
+			}
+
+			// Was told to add all UActorComponents?
+			const bool bIncludeAllSubobjects = Selection.IncludeAllOption == EMultiUserIncludeAllSubobjectsType::AllSubobjects;
+			const bool bIncludeDueToComponent = Subobject->IsA(UActorComponent::StaticClass())
+				&& Selection.IncludeAllOption == EMultiUserIncludeAllSubobjectsType::AllComponents;
+			if (bIncludeAllSubobjects || bIncludeDueToComponent)
+			{
+				FurtherObjectsCallback(*Subobject);
+				return;
+			}
+
+			// Has configured class?
+			UClass* CurrentClass = Subobject->GetClass();
+			for (; CurrentClass; CurrentClass = CurrentClass->GetSuperClass())
+			{
+				if (Selection.IncludeClasses.Contains(Subobject->GetClass()))
+				{
+					FurtherObjectsCallback(*Subobject);
+					return;
+				}
+			}
+
+			// Has inclusion regex?
+			const bool bShouldIncludeByRegex = MatchesAnyRegex(Subobject->GetName(), Selection.IncludeSubobjectRegex);
+			if (bShouldIncludeByRegex)
+			{
+				FurtherObjectsCallback(*Subobject);
+			}
+		}, bIncludeNested);
+	}
+
+	static void InternalAddAdditionalObjectsFromSettings(UClass& StartClass, const FConcertReplicationEditorSettings& Settings, UObject& AddedObject, TFunctionRef<void(UObject&)> FurtherObjectsCallback)
+	{
+		// Find the most specialized class properties
+		UClass* Current = &StartClass;
+		const FConcertDefaultSubobjectSelection* DefaultSubobjects = nullptr;
+		for (; Current && !DefaultSubobjects; Current = Current->GetSuperClass())
+		{
+			DefaultSubobjects = Settings.DefaultSubobjectSelection.Find(Current);
+			if (!DefaultSubobjects)
+			{
+				continue;
+			}
+			
+			ApplyDefaultSubobjectSelection(AddedObject, *DefaultSubobjects, FurtherObjectsCallback);
+			// Recurse super structs
+			if (UClass* Parent = Current->GetSuperClass()
+				; Parent && DefaultSubobjects->bInheritFromBase)
+			{
+				InternalAddAdditionalObjectsFromSettings(*Parent, Settings, AddedObject, FurtherObjectsCallback);
+			}
+		}
+	}
+}
+
+void FConcertReplicationEditorSettings::AddAdditionalObjectsFromSettings(UObject& AddedObject, TFunctionRef<void(UObject&)> FurtherObjectsCallback) const
+{
+	UE::ConcertClientSharedSlate::DefaultSubobjects::InternalAddAdditionalObjectsFromSettings(*AddedObject.GetClass(), *this, AddedObject, FurtherObjectsCallback);
+}
\ No newline at end of file
diff --git a/Engine/Plugins/Developer/Concert/ConcertUI/ConcertClientSharedSlate/Source/ConcertClientSharedSlate/Public/Replication/Editor/Model/Object/AddOuterSource.h b/Engine/Plugins/Developer/Concert/ConcertUI/ConcertClientSharedSlate/Source/ConcertClientSharedSlate/Public/Replication/Editor/Model/Object/AddOuterSource.h
new file mode 100644
--- /dev/null
+++ b/Engine/Plugins/Developer/Concert/ConcertUI/ConcertClientSharedSlate/Source/ConcertClientSharedSlate/Public/Replication/Editor/Model/Object/AddOuterSource.h
@@ -0,0 +1,29 @@
+// Copyright Epic Games, Inc. All Rights Reserved.
+
+#pragma once
+
+#include "Replication/Editor/Model/Object/IObjectSourceModel.h"
+
+namespace UE::ConcertClientSharedSlate
+{
+	/** Adds the outer object as object. Used as source in context menu. */
+	class CONCERTCLIENTSHAREDSLATE_API FAddOuterSource : public IObjectSourceModel
+	{
+	public:
+
+		bool TrySetObjectIfValid(UObject* InObject);
+
+		//~ Begin IObjectSourceModel Interface
+		virtual ConcertSharedSlate::FSourceDisplayInfo GetDisplayInfo() const override;
+		virtual uint32 GetNumSelectableItems() const override;
+		virtual void EnumerateSelectableItems(TFunctionRef<EBreakBehavior(const FSelectableObjectInfo& SelectableOption)> Delegate) const override;
+		//~ End IObjectSourceModel Interface
+
+	private:
+
+		/** The object whose outer to return. */
+		TWeakObjectPtr<UObject> Object;
+	};
+}
+
+
diff --git a/Engine/Plugins/Developer/Concert/ConcertUI/ConcertClientSharedSlate/Source/ConcertClientSharedSlate/Public/Replication/Editor/Model/Object/SelectedActorsSource.h b/Engine/Plugins/Developer/Concert/ConcertUI/ConcertClientSharedSlate/Source/ConcertClientSharedSlate/Public/Replication/Editor/Model/Object/SelectedActorsSource.h
new file mode 100644
--- /dev/null
+++ b/Engine/Plugins/Developer/Concert/ConcertUI/ConcertClientSharedSlate/Source/ConcertClientSharedSlate/Public/Replication/Editor/Model/Object/SelectedActorsSource.h
@@ -0,0 +1,20 @@
+// Copyright Epic Games, Inc. All Rights Reserved.
+
+#pragma once
+
+#include "Replication/Editor/Model/Object/IObjectSourceModel.h"
+
+namespace UE::ConcertClientSharedSlate
+{
+	/** Gets the actors currently selected in the editor. */
+	class CONCERTCLIENTSHAREDSLATE_API FSelectedActorsSource : public IObjectSourceModel
+	{
+	public:
+
+		//~ Begin IObjectSourceModel Interface
+		virtual ConcertSharedSlate::FSourceDisplayInfo GetDisplayInfo() const override;
+		virtual uint32 GetNumSelectableItems() const override;
+		virtual void EnumerateSelectableItems(TFunctionRef<EBreakBehavior(const FSelectableObjectInfo& SelectableOption)> Delegate) const override;
+		//~ End IObjectSourceModel Interface
+	};
+}
diff --git a/Engine/Plugins/Developer/Concert/ConcertUI/ConcertClientSharedSlate/Source/ConcertClientSharedSlate/Public/Replication/Editor/Model/Object/WorldActorSource.h b/Engine/Plugins/Developer/Concert/ConcertUI/ConcertClientSharedSlate/Source/ConcertClientSharedSlate/Public/Replication/Editor/Model/Object/WorldActorSource.h
new file mode 100644
--- /dev/null
+++ b/Engine/Plugins/Developer/Concert/ConcertUI/ConcertClientSharedSlate/Source/ConcertClientSharedSlate/Public/Replication/Editor/Model/Object/WorldActorSource.h
@@ -0,0 +1,22 @@
+// Copyright Epic Games, Inc. All Rights Reserved.
+
+#pragma once
+
+#include "Replication/Editor/Model/Object/IObjectSourceModel.h"
+
+namespace UE::ConcertClientSharedSlate
+{
+	/** Gets all the actors in the current GWorld */
+	class CONCERTCLIENTSHAREDSLATE_API FWorldActorSource : public IObjectSourceModel
+	{
+	public:
+
+		//~ Begin IObjectSourceModel Interface
+		virtual ConcertSharedSlate::FSourceDisplayInfo GetDisplayInfo() const override;
+		virtual uint32 GetNumSelectableItems() const override;
+		virtual void EnumerateSelectableItems(TFunctionRef<EBreakBehavior(const FSelectableObjectInfo& SelectableOption)> Delegate) const override;
+		//~ End IObjectSourceModel Interface
+	};
+}
+
+
diff --git a/Engine/Plugins/Developer/Concert/ConcertUI/ConcertClientSharedSlate/Source/ConcertClientSharedSlate/Public/Replication/ReplicationWidgetFactories.h b/Engine/Plugins/Developer/Concert/ConcertUI/ConcertClientSharedSlate/Source/ConcertClientSharedSlate/Public/Replication/ReplicationWidgetFactories.h
new file mode 100644
--- /dev/null
+++ b/Engine/Plugins/Developer/Concert/ConcertUI/ConcertClientSharedSlate/Source/ConcertClientSharedSlate/Public/Replication/ReplicationWidgetFactories.h
@@ -0,0 +1,59 @@
+// Copyright Epic Games, Inc. All Rights Reserved.
+
+#pragma once
+
+#include "Misc/Attribute.h"
+#include "Templates/SharedPointer.h"
+
+class UObject;
+class SWidget;
+
+struct FConcertReplicationEditorSettings;
+struct FObjectReplicationMap;
+
+namespace UE::ConcertClientSharedSlate
+{
+	class IEditableObjectToPropertiesModel;
+	class IObjectSelectionSourceModel;
+	class IPropertySelectionSourceModel;
+		
+	struct FCreateEditorParams
+	{
+		/** The model that the editor is displaying */
+		TSharedRef<IEditableObjectToPropertiesModel> DataModel;
+
+		/** Determines the objects that can be added to the object list. */
+		TSharedRef<IObjectSelectionSourceModel> ObjectSource;
+
+		/** Determines the properties that can be added to the property list. */
+		TSharedRef<IPropertySelectionSourceModel> PropertySource;
+
+		// TODO DP: Add way to add more columns
+	};
+
+	/**
+	 * Creates a replication editor.
+	 * 
+	 * The editor consists of three areas, which share similar workflows as the world outliner and details panel in the level editor:
+	 * 1. Root objects,  similar to world outliner: objects (usually actors) are added here. FCreateEditorParams::ObjectSource is used to build a combo button through which new objects can be added.
+	 * 2. Subobjects, similar to component view (SSubobjectEditor): Shows subobjects of a root objects selected above; typically components.
+	 * 3. Properties, similar to details panel: Shows properties of the selected root object and / or subobjects.
+	 */
+	CONCERTCLIENTSHAREDSLATE_API TSharedRef<SWidget> CreateEditor(FCreateEditorParams Params);
+
+	/**
+	 * Creates a model that can be passed to CreateEditor.
+	 * 
+	 * This model edits a FObjectReplicationMap that is assumed to be within the transactional OwnerObject. The model will respond to undo & redo by triggering the model's update callbacks.
+	 * 
+	 * @param OwnerObject The object containing the FObjectReplicationMap.
+	 * @param ReplicationMapAttribute Getter for extracting the FObjectReplicationMap to edit
+	 * 
+	 * @return A model that will edit the FObjectReplicationMap.
+	 */
+	CONCERTCLIENTSHAREDSLATE_API TSharedRef<IEditableObjectToPropertiesModel> CreatePropertySelectionModel(
+		UObject& OwnerObject,
+		TAttribute<FObjectReplicationMap*> ReplicationMapAttribute,
+		TAttribute<const FConcertReplicationEditorSettings*> OptionalReplicationSettingsAttribute = {}
+		);
+}
\ No newline at end of file
diff --git a/Engine/Plugins/Developer/Concert/ConcertUI/ConcertClientSharedSlate/Source/ConcertClientSharedSlate/Public/Replication/Settings/ConcertDefaultPropertySelection.h b/Engine/Plugins/Developer/Concert/ConcertUI/ConcertClientSharedSlate/Source/ConcertClientSharedSlate/Public/Replication/Settings/ConcertDefaultPropertySelection.h
new file mode 100644
--- /dev/null
+++ b/Engine/Plugins/Developer/Concert/ConcertUI/ConcertClientSharedSlate/Source/ConcertClientSharedSlate/Public/Replication/Settings/ConcertDefaultPropertySelection.h
@@ -0,0 +1,20 @@
+// Copyright Epic Games, Inc. All Rights Reserved.
+
+#pragma once
+
+#include "ConcertInheritableClassOption.h"
+#include "ConcertDefaultPropertySelection.generated.h"
+
+USTRUCT()
+struct FConcertDefaultPropertySelection : public FConcertInheritableClassOption
+{
+	GENERATED_BODY()
+
+	/**
+	 * A list of properties that should be selected by default when you add a new object type.
+	 * Specify properties by using the names as the appear in the replication editor and each property with ".".
+	 * Example: "RelativeLocation.X"
+	 */
+	UPROPERTY(EditAnywhere, Category = "Config")
+	TArray<FString> DefaultSelectedProperties;
+};
diff --git a/Engine/Plugins/Developer/Concert/ConcertUI/ConcertClientSharedSlate/Source/ConcertClientSharedSlate/Public/Replication/Settings/ConcertDefaultSubobjectSelection.h b/Engine/Plugins/Developer/Concert/ConcertUI/ConcertClientSharedSlate/Source/ConcertClientSharedSlate/Public/Replication/Settings/ConcertDefaultSubobjectSelection.h
new file mode 100644
--- /dev/null
+++ b/Engine/Plugins/Developer/Concert/ConcertUI/ConcertClientSharedSlate/Source/ConcertClientSharedSlate/Public/Replication/Settings/ConcertDefaultSubobjectSelection.h
@@ -0,0 +1,56 @@
+// Copyright Epic Games, Inc. All Rights Reserved.
+
+#pragma once
+
+#include "ConcertInheritableClassOption.h"
+#include "UObject/SoftObjectPath.h"
+#include "ConcertDefaultSubobjectSelection.generated.h"
+
+UENUM()
+enum class EMultiUserIncludeAllSubobjectsType : uint8
+{
+	/** Include nothing by default */
+	None,
+	/** Include all components */
+	AllComponents,
+	/** Include all subobjects recursively (includes components, their subobjects, etc. ) */
+	AllSubobjects
+};
+
+/** Settings for components that are supposed to be included  */
+USTRUCT()
+struct FConcertDefaultSubobjectSelection : public FConcertInheritableClassOption
+{
+	GENERATED_BODY()
+	
+	/**
+	 * A list of subobject classes to add by default.
+	 * For example, when you add an actor, you want to also add certain component classes.
+	 */
+	UPROPERTY(EditAnywhere, Category = "Config")
+	TSet<FSoftClassPath> IncludeClasses;
+
+	/**
+	 * Components that match any of this regex will be excluded from DefaultSelectedComponentClasses.
+	 *
+	 * Tips:
+	 * - If you want to include "Component" but not "ComponentName" you can use boundaries "\bComponent\b".
+	 * . ".*" matches any sequence of characters.
+	 */
+	UPROPERTY(EditAnywhere, Category = "Config")
+	TSet<FString> ExcludeSubobjectRegex;
+
+	/**
+	 * Components that match any of this regex will be included as well.
+	 * 
+	 * Tips:
+	 * - If you want to include "Component" but not "ComponentName" you can use boundaries "\bComponent\b".
+	 * . ".*" matches any sequence of characters. 
+	 */
+	UPROPERTY(EditAnywhere, Category = "Config")
+	TSet<FString> IncludeSubobjectRegex;
+
+	/** Behaviour to configure for including all of a certain type of subobject. */
+	UPROPERTY(EditAnywhere, AdvancedDisplay, Category = "Config")
+	EMultiUserIncludeAllSubobjectsType IncludeAllOption = EMultiUserIncludeAllSubobjectsType::None;
+};
diff --git a/Engine/Plugins/Developer/Concert/ConcertUI/ConcertClientSharedSlate/Source/ConcertClientSharedSlate/Public/Replication/Settings/ConcertReplicationEditorSettings.h b/Engine/Plugins/Developer/Concert/ConcertUI/ConcertClientSharedSlate/Source/ConcertClientSharedSlate/Public/Replication/Settings/ConcertReplicationEditorSettings.h
new file mode 100644
--- /dev/null
+++ b/Engine/Plugins/Developer/Concert/ConcertUI/ConcertClientSharedSlate/Source/ConcertClientSharedSlate/Public/Replication/Settings/ConcertReplicationEditorSettings.h
@@ -0,0 +1,34 @@
+// Copyright Epic Games, Inc. All Rights Reserved.
+
+#pragma once
+
+#include "ConcertDefaultPropertySelection.h"
+#include "ConcertDefaultSubobjectSelection.h"
+#include "ConcertReplicationEditorSettings.generated.h"
+
+struct FReplicatedObjectInfo;
+
+USTRUCT()
+struct CONCERTCLIENTSHAREDSLATE_API FConcertReplicationEditorSettings
+{
+	GENERATED_BODY()
+	
+	/** Properties you want selected by default when you add a new replicated object in the editor */
+	UPROPERTY(EditAnywhere, Config, Category = "Replication|Editor")
+	TMap<FSoftClassPath, FConcertDefaultPropertySelection> DefaultPropertySelection;
+
+	/**
+	 * Subobjects you want selected by default when you add a new replicated object in the editor.
+	 * For examples:
+	 *  - Whenever you add a StaticMeshActor actor, you its static mesh component to be added automatically, too.
+	 *  - Whenever you add a component, you want to add its child components, too.
+	 */
+	UPROPERTY(EditAnywhere, Config, Category = "Replication|Editor")
+	TMap<FSoftClassPath, FConcertDefaultSubobjectSelection> DefaultSubobjectSelection;
+
+	/** Reads DefaultPropertySelection and applies any default property selections to Info based on the Class just added. */
+	void AddDefaultPropertiesFromSettings(FReplicatedObjectInfo& Info, UClass& Class) const;
+
+	/** Reads DefaultComponentSelection and calls FurtherObjectsCallback on any further objects that should also be added. */
+	void AddAdditionalObjectsFromSettings(UObject& AddedObject, TFunctionRef<void(UObject&)> FurtherObjectsCallback) const;
+};
\ No newline at end of file
