Fix ConcertClientSharedSlate.uplugin:
- Remove MultiUserSlateServer from program allow list (this plugin is supposed to be used by CLIENTS)
- Add ConcertSharedSlate as dependency to ConcertClientSharedSlate

diff --git a/Engine/Plugins/Developer/Concert/ConcertUI/ConcertClientSharedSlate/ConcertClientSharedSlate.uplugin b/Engine/Plugins/Developer/Concert/ConcertUI/ConcertClientSharedSlate/ConcertClientSharedSlate.uplugin
--- a/Engine/Plugins/Developer/Concert/ConcertUI/ConcertClientSharedSlate/ConcertClientSharedSlate.uplugin
+++ b/Engine/Plugins/Developer/Concert/ConcertUI/ConcertClientSharedSlate/ConcertClientSharedSlate.uplugin
@@ -21,7 +21,7 @@
           "Name": "ConcertClientSharedSlate",
           "Type": "UncookedOnly",
           "LoadingPhase": "Default",
-          "ProgramAllowList": [ "UnrealMultiUserServer", "UnrealRecoverySvc", "CrashReportClientEditor" ]
+          "ProgramAllowList": [ "UnrealRecoverySvc", "CrashReportClientEditor" ]
         }
 	],
 	"Plugins": [
@@ -29,6 +29,10 @@
 			"Name": "ConcertMain",
 			"Enabled": true
 		},
+		{
+			"Name": "ConcertSharedSlate",
+			"Enabled": true
+		},
 		{
 			"Name": "ConcertSyncCore",
 			"Enabled": true
