Subobject editor for replication

diff --git a/Engine/Plugins/Developer/Concert/ConcertApp/MultiUserClient/Source/MultiUserClient/Private/Replication/Stream/ClientStreamRepository.cpp b/Engine/Plugins/Developer/Concert/ConcertApp/MultiUserClient/Source/MultiUserClient/Private/Replication/Stream/ClientStreamRepository.cpp
--- a/Engine/Plugins/Developer/Concert/ConcertApp/MultiUserClient/Source/MultiUserClient/Private/Replication/Stream/ClientStreamRepository.cpp
+++ b/Engine/Plugins/Developer/Concert/ConcertApp/MultiUserClient/Source/MultiUserClient/Private/Replication/Stream/ClientStreamRepository.cpp
@@ -39,8 +39,8 @@ namespace UE::MultiUserClient
 	}
 
 	void FClientStreamRepository::OnObjectsChanged(
-		TArrayView<UObject*> Objects,
-		TArrayView<FSoftObjectPath> SoftObjectPaths,
+		TConstArrayView<UObject*> Objects,
+		TConstArrayView<FSoftObjectPath> SoftObjectPaths,
 		ConcertClientSharedSlate::EReplicatedObjectChangeReason ReplicatedObjectChangeReason
 		)
 	{
diff --git a/Engine/Plugins/Developer/Concert/ConcertApp/MultiUserClient/Source/MultiUserClient/Private/Replication/Stream/ClientStreamRepository.h b/Engine/Plugins/Developer/Concert/ConcertApp/MultiUserClient/Source/MultiUserClient/Private/Replication/Stream/ClientStreamRepository.h
--- a/Engine/Plugins/Developer/Concert/ConcertApp/MultiUserClient/Source/MultiUserClient/Private/Replication/Stream/ClientStreamRepository.h
+++ b/Engine/Plugins/Developer/Concert/ConcertApp/MultiUserClient/Source/MultiUserClient/Private/Replication/Stream/ClientStreamRepository.h
@@ -74,7 +74,7 @@ namespace UE::MultiUserClient
 		/** Called when the data underlying the model has changed and the UI needs to be refreshed. */
 		FOnModelChanged OnModelChangedDelegate;
 		
-		void OnObjectsChanged(TArrayView<UObject*> Objects, TArrayView<FSoftObjectPath> SoftObjectPaths, ConcertClientSharedSlate::EReplicatedObjectChangeReason ReplicatedObjectChangeReason);
+		void OnObjectsChanged(TConstArrayView<UObject*> Objects, TConstArrayView<FSoftObjectPath> SoftObjectPaths, ConcertClientSharedSlate::EReplicatedObjectChangeReason ReplicatedObjectChangeReason);
 		void OnPropertiesChanged();
 	};
 }
diff --git a/Engine/Plugins/Developer/Concert/ConcertApp/MultiUserClient/Source/MultiUserClient/Private/Widgets/ActiveSession/Replication/Joined/SReplicationJoinedView.cpp b/Engine/Plugins/Developer/Concert/ConcertApp/MultiUserClient/Source/MultiUserClient/Private/Widgets/ActiveSession/Replication/Joined/SReplicationJoinedView.cpp
--- a/Engine/Plugins/Developer/Concert/ConcertApp/MultiUserClient/Source/MultiUserClient/Private/Widgets/ActiveSession/Replication/Joined/SReplicationJoinedView.cpp
+++ b/Engine/Plugins/Developer/Concert/ConcertApp/MultiUserClient/Source/MultiUserClient/Private/Widgets/ActiveSession/Replication/Joined/SReplicationJoinedView.cpp
@@ -32,7 +32,7 @@ namespace UE::MultiUserClient
 			MakeShared<FSelectPropertyFromUClassModel>()
 		};
 
-		const TSharedRef<IReplicationEditorView> NewEditorView = CreateEditor(ReplicationEditorCreationParams);
+		const TSharedRef<IReplicationEditorView> NewEditorView = CreateEditorForUnrealEditor(ReplicationEditorCreationParams);
 		WeakEditorView = NewEditorView;
 		ChildSlot
 		[
diff --git a/Engine/Plugins/Developer/Concert/ConcertApp/MultiUserClient/Source/MultiUserReplicationEditor/Private/StreamEditor/ReplicationSessionPresetEditorToolkit.cpp b/Engine/Plugins/Developer/Concert/ConcertApp/MultiUserClient/Source/MultiUserReplicationEditor/Private/StreamEditor/ReplicationSessionPresetEditorToolkit.cpp
--- a/Engine/Plugins/Developer/Concert/ConcertApp/MultiUserClient/Source/MultiUserReplicationEditor/Private/StreamEditor/ReplicationSessionPresetEditorToolkit.cpp
+++ b/Engine/Plugins/Developer/Concert/ConcertApp/MultiUserClient/Source/MultiUserReplicationEditor/Private/StreamEditor/ReplicationSessionPresetEditorToolkit.cpp
@@ -84,7 +84,7 @@ namespace UE::MultiUserReplicationEditor
 			}));
 		const TSharedRef<FEditorObjectSelectionSourceModel> ObjectSourceModel = MakeShared<FEditorObjectSelectionSourceModel>();
 		const TSharedRef<FSelectPropertyFromUClassModel> PropertySourceModel = MakeShared<FSelectPropertyFromUClassModel>();
-		const TSharedRef<IReplicationEditorView> EditorView = CreateEditor(FCreateEditorParams{ AssetReadWriteModel, ObjectSourceModel, PropertySourceModel });
+		const TSharedRef<IReplicationEditorView> EditorView = CreateEditorForUnrealEditor(FCreateEditorParams{ AssetReadWriteModel, ObjectSourceModel, PropertySourceModel });
 		return SNew(SDockTab)
 			.Label(LOCTEXT("BaseDetailsTitle", "Details"))
 			[
diff --git a/Engine/Plugins/Developer/Concert/ConcertUI/ConcertClientSharedSlate/Source/ConcertClientSharedSlate/ConcertClientSharedSlate.Build.cs b/Engine/Plugins/Developer/Concert/ConcertUI/ConcertClientSharedSlate/Source/ConcertClientSharedSlate/ConcertClientSharedSlate.Build.cs
--- a/Engine/Plugins/Developer/Concert/ConcertUI/ConcertClientSharedSlate/Source/ConcertClientSharedSlate/ConcertClientSharedSlate.Build.cs
+++ b/Engine/Plugins/Developer/Concert/ConcertUI/ConcertClientSharedSlate/Source/ConcertClientSharedSlate/ConcertClientSharedSlate.Build.cs
@@ -33,6 +33,8 @@ namespace UnrealBuildTool.Rules
 					"Projects",
 					"Slate",
 					"SlateCore",
+					"SubobjectDataInterface",
+					"SubobjectEditor",
 					"ToolMenus",
 					"ToolWidgets",
 					"UnrealEd", 
diff --git a/Engine/Plugins/Developer/Concert/ConcertUI/ConcertClientSharedSlate/Source/ConcertClientSharedSlate/Private/Replication/Editor/Model/GenericPropertySelectionModel.cpp b/Engine/Plugins/Developer/Concert/ConcertUI/ConcertClientSharedSlate/Source/ConcertClientSharedSlate/Private/Replication/Editor/Model/GenericPropertySelectionModel.cpp
--- a/Engine/Plugins/Developer/Concert/ConcertUI/ConcertClientSharedSlate/Source/ConcertClientSharedSlate/Private/Replication/Editor/Model/GenericPropertySelectionModel.cpp
+++ b/Engine/Plugins/Developer/Concert/ConcertUI/ConcertClientSharedSlate/Source/ConcertClientSharedSlate/Private/Replication/Editor/Model/GenericPropertySelectionModel.cpp
@@ -117,7 +117,7 @@ namespace UE::ConcertClientSharedSlate
 		return !AssignedProperties->PropertySelection.ReplicatedProperties.IsEmpty();
 	}
 
-	void FGenericPropertySelectionModel::AddObjects(TArrayView<UObject*> Objects)
+	void FGenericPropertySelectionModel::AddObjects(TConstArrayView<UObject*> Objects)
 	{
 		FObjectReplicationMap* ReplicationMap = ReplicationMapAttribute.Get();
 		if (!ensure(ReplicationMap) || Objects.IsEmpty())
@@ -176,7 +176,7 @@ namespace UE::ConcertClientSharedSlate
 		AddObjects(AdditionalObjectsToAdd);
 	}
 
-	void FGenericPropertySelectionModel::RemoveObjects(TArrayView<FSoftObjectPath> Objects)
+	void FGenericPropertySelectionModel::RemoveObjects(TConstArrayView<FSoftObjectPath> Objects)
 	{
 		FObjectReplicationMap* ReplicationMap = ReplicationMapAttribute.Get();
 		if (!ensure(ReplicationMap) || Objects.IsEmpty())
@@ -213,7 +213,7 @@ namespace UE::ConcertClientSharedSlate
 		}
 	}
 
-	void FGenericPropertySelectionModel::AddProperties(const FSoftObjectPath& Object, TArrayView<FConcertPropertyChain> Properties)
+	void FGenericPropertySelectionModel::AddProperties(const FSoftObjectPath& Object, TConstArrayView<FConcertPropertyChain> Properties)
 	{
 		FObjectReplicationMap* ReplicationMap = ReplicationMapAttribute.Get();
 		if (!ensure(ReplicationMap))
@@ -258,7 +258,7 @@ namespace UE::ConcertClientSharedSlate
 		}
 	}
 
-	void FGenericPropertySelectionModel::RemoveProperties(const FSoftObjectPath& Object, TArrayView<FConcertPropertyChain> Properties)
+	void FGenericPropertySelectionModel::RemoveProperties(const FSoftObjectPath& Object, TConstArrayView<FConcertPropertyChain> Properties)
 	{
 		FObjectReplicationMap* ReplicationMap = ReplicationMapAttribute.Get();
 		if (!ensure(ReplicationMap))
diff --git a/Engine/Plugins/Developer/Concert/ConcertUI/ConcertClientSharedSlate/Source/ConcertClientSharedSlate/Private/Replication/Editor/Model/GenericPropertySelectionModel.h b/Engine/Plugins/Developer/Concert/ConcertUI/ConcertClientSharedSlate/Source/ConcertClientSharedSlate/Private/Replication/Editor/Model/GenericPropertySelectionModel.h
--- a/Engine/Plugins/Developer/Concert/ConcertUI/ConcertClientSharedSlate/Source/ConcertClientSharedSlate/Private/Replication/Editor/Model/GenericPropertySelectionModel.h
+++ b/Engine/Plugins/Developer/Concert/ConcertUI/ConcertClientSharedSlate/Source/ConcertClientSharedSlate/Private/Replication/Editor/Model/GenericPropertySelectionModel.h
@@ -35,10 +35,10 @@ namespace UE::ConcertClientSharedSlate
 		//~ End IObjectToPropertiesModel Interface
 		
 		//~ Begin IEditableObjectToPropertiesModel Interface
-		virtual void AddObjects(TArrayView<UObject*> Objects) override;
-		virtual void RemoveObjects(TArrayView<FSoftObjectPath> Objects) override;
-		virtual void AddProperties(const FSoftObjectPath&, TArrayView<FConcertPropertyChain> Properties) override;
-		virtual void RemoveProperties(const FSoftObjectPath&, TArrayView<FConcertPropertyChain> Properties) override;
+		virtual void AddObjects(TConstArrayView<UObject*> Objects) override;
+		virtual void RemoveObjects(TConstArrayView<FSoftObjectPath> Objects) override;
+		virtual void AddProperties(const FSoftObjectPath&, TConstArrayView<FConcertPropertyChain> Properties) override;
+		virtual void RemoveProperties(const FSoftObjectPath&, TConstArrayView<FConcertPropertyChain> Properties) override;
 		virtual FOnObjectsChanged& OnObjectsChanged() override { return OnObjectsChangedDelegate; }
 		virtual FOnPropertiesChanged& OnPropertiesChanged() override { return OnPropertiesChangedDelegate; }
 		//~ End IEditableObjectToPropertiesModel Interface
diff --git a/Engine/Plugins/Developer/Concert/ConcertUI/ConcertClientSharedSlate/Source/ConcertClientSharedSlate/Private/Replication/Editor/Model/Object/AddOuterSource.cpp b/Engine/Plugins/Developer/Concert/ConcertUI/ConcertClientSharedSlate/Source/ConcertClientSharedSlate/Private/Replication/Editor/Model/Object/AddOuterSource.cpp
deleted file mode 100644
diff --git a/Engine/Plugins/Developer/Concert/ConcertUI/ConcertClientSharedSlate/Source/ConcertClientSharedSlate/Private/Replication/Editor/Model/Object/ComponentFromActorSource.cpp b/Engine/Plugins/Developer/Concert/ConcertUI/ConcertClientSharedSlate/Source/ConcertClientSharedSlate/Private/Replication/Editor/Model/Object/ComponentFromActorSource.cpp
deleted file mode 100644
diff --git a/Engine/Plugins/Developer/Concert/ConcertUI/ConcertClientSharedSlate/Source/ConcertClientSharedSlate/Private/Replication/Editor/Model/Object/EditorObjectSelectionSourceModel.cpp b/Engine/Plugins/Developer/Concert/ConcertUI/ConcertClientSharedSlate/Source/ConcertClientSharedSlate/Private/Replication/Editor/Model/Object/EditorObjectSelectionSourceModel.cpp
--- a/Engine/Plugins/Developer/Concert/ConcertUI/ConcertClientSharedSlate/Source/ConcertClientSharedSlate/Private/Replication/Editor/Model/Object/EditorObjectSelectionSourceModel.cpp
+++ b/Engine/Plugins/Developer/Concert/ConcertUI/ConcertClientSharedSlate/Source/ConcertClientSharedSlate/Private/Replication/Editor/Model/Object/EditorObjectSelectionSourceModel.cpp
@@ -2,8 +2,6 @@
 
 #include "Replication/Editor/Model/Object/EditorObjectSelectionSourceModel.h"
 
-#include "Replication/Editor/Model/Object/AddOuterSource.h"
-#include "Replication/Editor/Model/Object/ComponentFromActorSource.h"
 #include "Replication/Editor/Model/Object/SelectedActorsSource.h"
 #include "Replication/Editor/Model/Object/WorldActorSource.h"
 
@@ -14,9 +12,7 @@
 namespace UE::ConcertClientSharedSlate
 {
 	FEditorObjectSelectionSourceModel::FEditorObjectSelectionSourceModel()
-		: WorldActorSource(MakeShared<FWorldActorSource>())
-		// Important: We dependent on the fact that WorldActorSource is declared (in the header) and hence initialized before BaseActorCategory
-		, BaseActorCategory([this]()
+		: BaseActorCategory([this]()
 		{
 			return FObjectSourceCategory
 			{
@@ -26,7 +22,7 @@ namespace UE::ConcertClientSharedSlate
 					LOCTEXT("ObjectsCategory.Tooltip", "Options for adding objects from the open editor world")
 				},
 				{
-					WorldActorSource,
+					MakeShared<FWorldActorSource>(),
 					MakeShared<FSelectedActorsSource>()
 				}
 			};
@@ -36,13 +32,12 @@ namespace UE::ConcertClientSharedSlate
 	TArray<FObjectSourceCategory> FEditorObjectSelectionSourceModel::GetRootSources() const
 	{
 		FObjectSourceCategory ActorCategory = BaseActorCategory;
-		FComponentFromActorSource_Root::AddComponentsAdditionToRootCategory(WorldActorSource.Get(), ActorCategory);
 		return { ActorCategory };
 	}
 
 	TArray<TSharedRef<IObjectSourceModel>> FEditorObjectSelectionSourceModel::GetContextMenuOptions(const FSoftObjectPath& Item)
 	{
-		return ContextMenuSources.BuildContextSources(Item);
+		return {};
 	}
 
 	EObjectItemValidity FEditorObjectSelectionSourceModel::GetItemValidity(const FSoftObjectPath& Item) const
@@ -54,32 +49,6 @@ namespace UE::ConcertClientSharedSlate
 		
 		return EObjectItemValidity::Invalid;
 	}
-
-	FEditorObjectSelectionSourceModel::FContextMenuSources::FContextMenuSources()
-		: AddComponent(MakeShared<FComponentFromActorSource_ContextMenu>())
-		, AddOuter(MakeShared<FAddOuterSource>())
-	{}
-
-	TArray<TSharedRef<IObjectSourceModel>> FEditorObjectSelectionSourceModel::FContextMenuSources::BuildContextSources(const FSoftObjectPath& Item)
-	{
-		UObject* ResolvedContextObject = Item.ResolveObject();
-		AActor* ResolvedContextActor = Cast<AActor>(ResolvedContextObject);
-		TArray<TSharedRef<IObjectSourceModel>> Result;
-
-		// The "Add Outer" option is not valid for all cases. For example adding an actor's outer is not valid.
-		if (AddOuter->TrySetObjectIfValid(ResolvedContextObject))
-		{
-			Result.Add(AddOuter);
-		}
-
-		// The "Add component" option is only valid if the context object is an actor displayed in the scene outliner
-		if (AddComponent->TrySetActorIfValid(ResolvedContextActor))
-		{
-			Result.Add(AddComponent);
-		}
-
-		return Result;
-	}
 }
 
 #undef LOCTEXT_NAMESPACE
\ No newline at end of file
diff --git a/Engine/Plugins/Developer/Concert/ConcertUI/ConcertClientSharedSlate/Source/ConcertClientSharedSlate/Private/Replication/Editor/Model/ReplicatedObjectData.h b/Engine/Plugins/Developer/Concert/ConcertUI/ConcertClientSharedSlate/Source/ConcertClientSharedSlate/Private/Replication/Editor/Model/ReplicatedObjectData.h
--- a/Engine/Plugins/Developer/Concert/ConcertUI/ConcertClientSharedSlate/Source/ConcertClientSharedSlate/Private/Replication/Editor/Model/ReplicatedObjectData.h
+++ b/Engine/Plugins/Developer/Concert/ConcertUI/ConcertClientSharedSlate/Source/ConcertClientSharedSlate/Private/Replication/Editor/Model/ReplicatedObjectData.h
@@ -6,7 +6,7 @@
 
 namespace UE::ConcertClientSharedSlate
 {
-	/** Instanced for each object row in SPropertyReplicationSelectionViewer.*/
+	/** Instanced for each object row in SObjectToPropertyViewer.*/
 	class FReplicatedObjectData
 	{
 	public:
@@ -24,10 +24,10 @@ namespace UE::ConcertClientSharedSlate
 		FSoftObjectPath ObjectPath;
 	};
 
-	/** Instanced for each object row SPropertyReplicationSelectionEditor. */
+	/** Instanced for each object row SObjectToPropertyEditor. */
 	class FReplicatedObjectData_Editor : public FReplicatedObjectData
 	{
-		// Empty and unused for now. I just want to highlight that if you want custom data for SPropertyReplicationSelectionEditor,
-		// you should add it here and override SPropertyReplicationSelectionViewer::AllocateObjectData.
+		// Empty and unused for now. I just want to highlight that if you want custom data for SObjectToPropertyEditor,
+		// you should add it here and override SObjectToPropertyViewer::AllocateObjectData.
 	};
 }
\ No newline at end of file
diff --git a/Engine/Plugins/Developer/Concert/ConcertUI/ConcertClientSharedSlate/Source/ConcertClientSharedSlate/Private/Replication/Editor/Model/ReplicatedPropertyData.h b/Engine/Plugins/Developer/Concert/ConcertUI/ConcertClientSharedSlate/Source/ConcertClientSharedSlate/Private/Replication/Editor/Model/ReplicatedPropertyData.h
--- a/Engine/Plugins/Developer/Concert/ConcertUI/ConcertClientSharedSlate/Source/ConcertClientSharedSlate/Private/Replication/Editor/Model/ReplicatedPropertyData.h
+++ b/Engine/Plugins/Developer/Concert/ConcertUI/ConcertClientSharedSlate/Source/ConcertClientSharedSlate/Private/Replication/Editor/Model/ReplicatedPropertyData.h
@@ -6,7 +6,7 @@
 
 namespace UE::ConcertClientSharedSlate
 {
-	/** Instanced for each property row in SPropertyReplicationSelectionViewer.*/
+	/** Instanced for each property row in SObjectToPropertyViewer.*/
 	class FReplicatedPropertyData
 	{
 	public:
@@ -28,10 +28,10 @@ namespace UE::ConcertClientSharedSlate
 		FConcertPropertyChain Property;
 	};
 
-	/** Instanced for each property row SPropertyReplicationSelectionEditor. */
+	/** Instanced for each property row SObjectToPropertyEditor. */
 	class FReplicatedPropertyData_Editor : public FReplicatedPropertyData
 	{
-		// Empty and unused for now. I just want to highlight that if you want custom data for SPropertyReplicationSelectionEditor,
-		// you should add it here and override SPropertyReplicationSelectionViewer::AllocatePropertyData.
+		// Empty and unused for now. I just want to highlight that if you want custom data for SObjectToPropertyEditor,
+		// you should add it here and override SObjectToPropertyViewer::AllocatePropertyData.
 	};
 }
\ No newline at end of file
diff --git a/Engine/Plugins/Developer/Concert/ConcertUI/ConcertClientSharedSlate/Source/ConcertClientSharedSlate/Private/Replication/Editor/Model/TransactionalPropertySelectionModel.cpp b/Engine/Plugins/Developer/Concert/ConcertUI/ConcertClientSharedSlate/Source/ConcertClientSharedSlate/Private/Replication/Editor/Model/TransactionalPropertySelectionModel.cpp
--- a/Engine/Plugins/Developer/Concert/ConcertUI/ConcertClientSharedSlate/Source/ConcertClientSharedSlate/Private/Replication/Editor/Model/TransactionalPropertySelectionModel.cpp
+++ b/Engine/Plugins/Developer/Concert/ConcertUI/ConcertClientSharedSlate/Source/ConcertClientSharedSlate/Private/Replication/Editor/Model/TransactionalPropertySelectionModel.cpp
@@ -17,28 +17,28 @@ namespace UE::ConcertClientSharedSlate
 		, OwningObject(&OwningObject)
 	{}
 
-	void FTransactionalPropertySelectionModel::AddObjects(TArrayView<UObject*> Objects)
+	void FTransactionalPropertySelectionModel::AddObjects(TConstArrayView<UObject*> Objects)
 	{
 		const FScopedTransaction Transaction(LOCTEXT("AddObjects", "Add replicated objects"));
 		OwningObject->Modify();
 		FGenericPropertySelectionModel::AddObjects(Objects);
 	}
 
-	void FTransactionalPropertySelectionModel::RemoveObjects(TArrayView<FSoftObjectPath> Objects)
+	void FTransactionalPropertySelectionModel::RemoveObjects(TConstArrayView<FSoftObjectPath> Objects)
 	{
 		const FScopedTransaction Transaction(LOCTEXT("RemoveObjects", "Remove replicated objects"));
 		OwningObject->Modify();
 		FGenericPropertySelectionModel::RemoveObjects(Objects);
 	}
 
-	void FTransactionalPropertySelectionModel::AddProperties(const FSoftObjectPath& SoftObjectPath, TArrayView<FConcertPropertyChain> Properties)
+	void FTransactionalPropertySelectionModel::AddProperties(const FSoftObjectPath& SoftObjectPath, TConstArrayView<FConcertPropertyChain> Properties)
 	{
 		const FScopedTransaction Transaction(LOCTEXT("AddProperties", "Add replicated properties"));
 		OwningObject->Modify();
 		FGenericPropertySelectionModel::AddProperties(SoftObjectPath, Properties);
 	}
 
-	void FTransactionalPropertySelectionModel::RemoveProperties(const FSoftObjectPath& SoftObjectPath, TArrayView<FConcertPropertyChain> Properties)
+	void FTransactionalPropertySelectionModel::RemoveProperties(const FSoftObjectPath& SoftObjectPath, TConstArrayView<FConcertPropertyChain> Properties)
 	{
 		const FScopedTransaction Transaction(LOCTEXT("RemoveProperties", "Remove replicated properties"));
 		OwningObject->Modify();
diff --git a/Engine/Plugins/Developer/Concert/ConcertUI/ConcertClientSharedSlate/Source/ConcertClientSharedSlate/Private/Replication/Editor/Model/TransactionalPropertySelectionModel.h b/Engine/Plugins/Developer/Concert/ConcertUI/ConcertClientSharedSlate/Source/ConcertClientSharedSlate/Private/Replication/Editor/Model/TransactionalPropertySelectionModel.h
--- a/Engine/Plugins/Developer/Concert/ConcertUI/ConcertClientSharedSlate/Source/ConcertClientSharedSlate/Private/Replication/Editor/Model/TransactionalPropertySelectionModel.h
+++ b/Engine/Plugins/Developer/Concert/ConcertUI/ConcertClientSharedSlate/Source/ConcertClientSharedSlate/Private/Replication/Editor/Model/TransactionalPropertySelectionModel.h
@@ -21,10 +21,10 @@ namespace UE::ConcertClientSharedSlate
 			);
 		
 		//~ Begin IEditableObjectToPropertiesModel Interface
-		virtual void AddObjects(TArrayView<UObject*> Objects) override;
-		virtual void RemoveObjects(TArrayView<FSoftObjectPath> Objects) override;
-		virtual void AddProperties(const FSoftObjectPath&, TArrayView<FConcertPropertyChain> Properties) override;
-		virtual void RemoveProperties(const FSoftObjectPath&, TArrayView<FConcertPropertyChain> Properties) override;
+		virtual void AddObjects(TConstArrayView<UObject*> Objects) override;
+		virtual void RemoveObjects(TConstArrayView<FSoftObjectPath> Objects) override;
+		virtual void AddProperties(const FSoftObjectPath&, TConstArrayView<FConcertPropertyChain> Properties) override;
+		virtual void RemoveProperties(const FSoftObjectPath&, TConstArrayView<FConcertPropertyChain> Properties) override;
 		//~ End IEditableObjectToPropertiesModel Interface
 
 		//~ Begin FEditorUndoClient Interface
diff --git a/Engine/Plugins/Developer/Concert/ConcertUI/ConcertClientSharedSlate/Source/ConcertClientSharedSlate/Private/Replication/Editor/View/ObjectEditor/FakeObjectToPropertiesEditorModel.cpp b/Engine/Plugins/Developer/Concert/ConcertUI/ConcertClientSharedSlate/Source/ConcertClientSharedSlate/Private/Replication/Editor/View/ObjectEditor/FakeObjectToPropertiesEditorModel.cpp
--- a/Engine/Plugins/Developer/Concert/ConcertUI/ConcertClientSharedSlate/Source/ConcertClientSharedSlate/Private/Replication/Editor/View/ObjectEditor/FakeObjectToPropertiesEditorModel.cpp
+++ b/Engine/Plugins/Developer/Concert/ConcertUI/ConcertClientSharedSlate/Source/ConcertClientSharedSlate/Private/Replication/Editor/View/ObjectEditor/FakeObjectToPropertiesEditorModel.cpp
@@ -6,15 +6,62 @@
 #include "Replication/Editor/Model/Property/IPropertySourceModel.h"
 #include "Replication/PropertyChainUtils.h"
 
+#include "GameFramework/Actor.h"
+
 namespace UE::ConcertClientSharedSlate
 {
-	bool FFakeObjectToPropertiesEditorModel::ForEachProperty(const FSoftObjectPath& Object, TFunctionRef<EBreakBehavior(const FConcertPropertyChain& Property)> Delegate) const
+	FSoftClassPath FFakeObjectToPropertiesEditorModel::GetObjectClass(const FSoftObjectPath& Object) const
+	{
+		const FSoftClassPath ResolvedClass = RealModel->GetObjectClass(Object);
+		if (ResolvedClass.IsValid())
+		{
+			return ResolvedClass;
+		}
+
+		// The object is not yet in the model.
+		const UObject* LoadedObject = Object.ResolveObject();
+		return LoadedObject ? LoadedObject->GetClass() : nullptr;
+	}
+
+	bool FFakeObjectToPropertiesEditorModel::ForEachReplicatedObject(TFunctionRef<EBreakBehavior(const FSoftObjectPath& Object)> Delegate) const
 	{
-		PropertySelectionSource->GetPropertySource(GetObjectClass(Object))
+		// This makes SObjectToPropertyViewer only show actors in the outliner
+		return RealModel->ForEachReplicatedObject([this, &Delegate](const FSoftObjectPath& Object)
+		{
+			return IsRootObject(Object)
+				? Delegate(Object)
+				: EBreakBehavior::Continue;
+		});
+	}
+
+	bool FFakeObjectToPropertiesEditorModel::ForEachProperty(const FSoftObjectPath& ObjectPath, TFunctionRef<EBreakBehavior(const FConcertPropertyChain& Property)> Delegate) const
+	{
+		// GetObjectClass might return null because ObjectPath might not be contained in the model
+		const FSoftClassPath ClassPath = GetObjectClass(ObjectPath);
+		UClass* LoadedClass = ClassPath.TryLoadClass<UObject>();
+		if (!LoadedClass)
+		{
+			UObject* Object = ObjectPath.ResolveObject();
+			LoadedClass = Object ? Object->GetClass() : nullptr;
+		}
+		if (!LoadedClass)
+		{
+			return false;
+		}
+		
+		// This makes SObjectToPropertyViewer list all available properties
+		PropertySelectionSource->GetPropertySource(LoadedClass)
 			->EnumerateSelectableItems([&Delegate](const FSelectablePropertyInfo& PropertyInfo)
 			{
 				return Delegate(PropertyInfo.Property);
 			});
 		return true;
 	}
+
+	bool FFakeObjectToPropertiesEditorModel::IsRootObject(const FSoftObjectPath& ObjectPath) const
+	{
+		const FSoftClassPath ClassPath = GetObjectClass(ObjectPath);
+		UClass* LoadedClass = ClassPath.TryLoadClass<UObject>();
+		return LoadedClass && LoadedClass->IsChildOf(AActor::StaticClass());
+	}
 }
diff --git a/Engine/Plugins/Developer/Concert/ConcertUI/ConcertClientSharedSlate/Source/ConcertClientSharedSlate/Private/Replication/Editor/View/ObjectEditor/FakeObjectToPropertiesEditorModel.h b/Engine/Plugins/Developer/Concert/ConcertUI/ConcertClientSharedSlate/Source/ConcertClientSharedSlate/Private/Replication/Editor/View/ObjectEditor/FakeObjectToPropertiesEditorModel.h
--- a/Engine/Plugins/Developer/Concert/ConcertUI/ConcertClientSharedSlate/Source/ConcertClientSharedSlate/Private/Replication/Editor/View/ObjectEditor/FakeObjectToPropertiesEditorModel.h
+++ b/Engine/Plugins/Developer/Concert/ConcertUI/ConcertClientSharedSlate/Source/ConcertClientSharedSlate/Private/Replication/Editor/View/ObjectEditor/FakeObjectToPropertiesEditorModel.h
@@ -7,14 +7,18 @@
 namespace UE::ConcertClientSharedSlate
 {
 	class IPropertySelectionSourceModel;
+	
 	/**
-	 * This model is used by SPropertyReplicationSelectionEditor.
+	 * This model is passed tp SObjectToPropertyView by SObjectToPropertyEditor.
+	 *
+	 * The goal here is to mock some of the functions so the UI displays more than is in the model for an easier editing
+	 * experience:
+	 * - The root object outliner will only displays AActors.
+	 * - The property view will display the properties reported by an IPropertySelectionSourceModel (which is usually ALL properties of the object class).
 	 * 
-	 * It forwards most functions except for ForEachProperty, which fakes to the SPropertyReplicationSelectionView
-	 * that all properties are selected. The properties are added to the real model via checking and unchecking the box.
-	 * 
-	 * This allows the SPropertyReplicationSelectionView to continue to show the properties it is being told are in the model,
-	 * which is important e.g. for the connection view (when a client is connected to a session and views the properties being sent).
+	 * SObjectToPropertyEditor injects column checkboxes on property rows which handle adding the properties to the real,
+	 * underlying model.
+	 * @see UE::ConcertClientSharedSlate::ReplicationPropertyColumns::ReplicatesColumns.
 	 */
 	class FFakeObjectToPropertiesEditorModel : public IObjectToPropertiesModel
 	{
@@ -26,14 +30,17 @@ namespace UE::ConcertClientSharedSlate
 		{}
 
 		//~ Begin IObjectToPropertiesModel Interface
+		
+		// Technically these functions should be also be wrapped but the SObjectToPropertyView does not use them so let's not for now.
 		virtual uint32 GetNumReplicatedObjects() const override { return RealModel->GetNumReplicatedObjects(); }
 		virtual uint32 GetNumProperties(const FSoftObjectPath& Object) const override { return RealModel->GetNumProperties(Object); }
-		virtual FSoftClassPath GetObjectClass(const FSoftObjectPath& Object) const override { return RealModel->GetObjectClass(Object); }
 		virtual bool ContainsObjects(const TSet<FSoftObjectPath>& Objects) const override { return RealModel->ContainsObjects(Objects); }
-		// This function is by SPropertyReplicationSelectionEditor and not used by SPropertyReplicationSelectionView so forward it normally
+		
+		virtual FSoftClassPath GetObjectClass(const FSoftObjectPath& Object) const override;
+		// This function is by SObjectToPropertyEditor and not used by SObjectToPropertyView so forward it normally
 		virtual bool ContainsProperties(const FSoftObjectPath& Object, const TSet<FConcertPropertyChain>& Properties) const override { return RealModel->ContainsProperties(Object, Properties); }
-		virtual bool ForEachReplicatedObject(TFunctionRef<EBreakBehavior(const FSoftObjectPath& Object)> Delegate) const override { return RealModel->ForEachReplicatedObject(Delegate); }
-		virtual bool ForEachProperty(const FSoftObjectPath& Object, TFunctionRef<EBreakBehavior(const FConcertPropertyChain& Property)> Delegate) const override;
+		virtual bool ForEachReplicatedObject(TFunctionRef<EBreakBehavior(const FSoftObjectPath& Object)> Delegate) const override;
+		virtual bool ForEachProperty(const FSoftObjectPath& ObjectPath, TFunctionRef<EBreakBehavior(const FConcertPropertyChain& Property)> Delegate) const override;
 		//~ End IObjectToPropertiesModel Interface
 
 	private:
@@ -43,6 +50,8 @@ namespace UE::ConcertClientSharedSlate
 
 		/** Determines the properties that can be selected. */
 		TSharedRef<IPropertySelectionSourceModel> PropertySelectionSource;
+
+		bool IsRootObject(const FSoftObjectPath& ObjectPath) const;
 	};
 }
 
diff --git a/Engine/Plugins/Developer/Concert/ConcertUI/ConcertClientSharedSlate/Source/ConcertClientSharedSlate/Private/Replication/Editor/View/ObjectEditor/SObjectToPropertyEditor.cpp b/Engine/Plugins/Developer/Concert/ConcertUI/ConcertClientSharedSlate/Source/ConcertClientSharedSlate/Private/Replication/Editor/View/ObjectEditor/SObjectToPropertyEditor.cpp
new file mode 100644
--- /dev/null
+++ b/Engine/Plugins/Developer/Concert/ConcertUI/ConcertClientSharedSlate/Source/ConcertClientSharedSlate/Private/Replication/Editor/View/ObjectEditor/SObjectToPropertyEditor.cpp
@@ -0,0 +1,257 @@
+// Copyright Epic Games, Inc. All Rights Reserved.
+
+#include "SObjectToPropertyEditor.h"
+
+#include "FakeObjectToPropertiesEditorModel.h"
+#include "Model/Item/SourceModelBuilders.h"
+#include "Replication/Editor/Model/DisplayUtils.h"
+#include "Replication/Editor/Model/IEditableObjectToPropertiesModel.h"
+#include "Replication/Editor/Model/Object/IObjectSelectionSourceModel.h"
+#include "Replication/Editor/Model/Property/IPropertySelectionSourceModel.h"
+#include "Replication/Editor/Model/ReplicatedPropertyData.h"
+#include "Replication/Editor/Model/ReplicatedObjectData.h"
+
+#include "Algo/AnyOf.h"
+#include "GameFramework/Actor.h"
+#include "UObject/Class.h"
+#include "Widgets/SBoxPanel.h"
+
+#define LOCTEXT_NAMESPACE "SObjectToPropertyEditor"
+
+namespace UE::ConcertClientSharedSlate
+{
+	void SObjectToPropertyEditor::Construct(const FArguments& InArgs,
+		TSharedRef<IEditableObjectToPropertiesModel> InPropertiesModel,
+		TSharedRef<IObjectSelectionSourceModel> InObjectSelectionSource,
+		TSharedRef<IPropertySelectionSourceModel> InPropertySelectionSource)
+	{
+		ObjectSelectionSource = MoveTemp(InObjectSelectionSource);
+		PropertySelectionSource = MoveTemp(InPropertySelectionSource);
+		
+		EditablePropertiesModel = MoveTemp(InPropertiesModel);
+		EditablePropertiesModel->OnObjectsChanged().AddSP(this, &SObjectToPropertyEditor::OnObjectsChanged);
+		EditablePropertiesModel->OnPropertiesChanged().AddSP(this, &SObjectToPropertyEditor::OnPropertiesChanged);
+		PropertiesModelAdapter = MakeShared<FFakeObjectToPropertiesEditorModel>(EditablePropertiesModel.ToSharedRef(), PropertySelectionSource.ToSharedRef());
+
+		using namespace ReplicationPropertyColumns;
+		const FReplicationPropertyColumn ReplicatesColumn =
+			ReplicatesColumns(
+				FGetPropertyCheckboxState::CreateSP(this, &SObjectToPropertyEditor::OnGetPropertyCheckboxState),
+				FOnPropertyCheckboxChanged::CreateSP(this, &SObjectToPropertyEditor::OnPropertyCheckboxChanged)
+				);
+		
+		SObjectToPropertyView::Construct(
+			SObjectToPropertyView::FArguments()
+				.AdditionalPropertyColumns({ ReplicatesColumn })
+				.SubobjectView(InArgs._SubobjectView)
+				.OnDeleteObjects(this, &SObjectToPropertyEditor::OnDeleteObjects)
+				.OnObjectsContextMenuOpening(this, &SObjectToPropertyEditor::OnObjectsContextMenuOpening)
+				.SortPropertyRowPredicate(this, &SObjectToPropertyEditor::SortBySelectionThenByName_PropertyPredicate)
+				.LeftOfObjectSearchBar()
+				[
+					BuildRootAddObjectWidgets()
+				]
+				.LeftOfPropertySearchBar()
+				[
+					SAssignNew(AddPropertyWidgetContainer, SHorizontalBox)
+				],
+			PropertiesModelAdapter.ToSharedRef()
+			);
+	}
+
+	void SObjectToPropertyEditor::OnObjectsChanged(TConstArrayView<UObject*> AddedObjects, TConstArrayView<FSoftObjectPath> RemovedObjects, EReplicatedObjectChangeReason ChangeReason)
+	{
+		RefreshObjectData();
+	}
+
+	void SObjectToPropertyEditor::OnPropertiesChanged()
+	{
+		RefreshPropertyData();
+	}
+
+	TSharedRef<SWidget> SObjectToPropertyEditor::BuildRootAddObjectWidgets() const
+	{
+		using namespace ConcertSharedSlate;
+		
+		const TSharedRef<SHorizontalBox> Root = SNew(SHorizontalBox);
+		const FSourceModelBuilders<FSelectableObjectInfo>::FItemPickerArgs Args = MakeObjectSourceBuilderArgs();
+		for (const TSourceSelectionCategory<FSelectableObjectInfo>& Category : ObjectSelectionSource->GetRootSources())
+		{
+			Root->AddSlot()
+				.AutoWidth()
+				[
+					FSourceModelBuilders<FSelectableObjectInfo>::BuildCategory(Category, Args)	
+				];
+		}
+		return Root;
+	}
+
+	void SObjectToPropertyEditor::OnObjectsSelectedForAdding(TArray<FSelectableObjectInfo> ObjectsToAdd) const
+	{
+		TArray<UObject*> Objects;
+		Algo::TransformIf(
+			ObjectsToAdd,
+			Objects,
+			[](const FSelectableObjectInfo& SelectableObject) { return SelectableObject.Object.IsValid(); },
+			[](const FSelectableObjectInfo& SelectableObject){ return SelectableObject.Object.Get(); }
+			);
+		
+		EditablePropertiesModel->AddObjects(Objects);
+	}
+
+	ECheckBoxState SObjectToPropertyEditor::OnGetPropertyCheckboxState(const FConcertPropertyChain& PropertyChain)
+	{
+		const TArray<FSoftObjectPath> SelectedObjectPaths = GetSelectedObjectShowingProperties();
+		return ReplicationPropertyColumns::GetPropertyCheckboxStateBasedOnSelection(PropertyChain, SelectedObjectPaths, *EditablePropertiesModel);
+	}
+
+	void SObjectToPropertyEditor::OnPropertyCheckboxChanged(bool bIsChecked, const FConcertPropertyChain& PropertyChain)
+	{
+		TArray Properties{ PropertyChain };
+		for (const FSoftObjectPath& Path : GetSelectedObjectShowingProperties())
+		{
+			UObject* Object = Path.ResolveObject();
+			if (!Object)
+			{
+				continue;
+			}
+			
+			if (bIsChecked)
+			{
+				if (!EditablePropertiesModel->ContainsObjects({ Path }))
+				{
+					EditablePropertiesModel->AddObjects({ Object });
+				}
+				EditablePropertiesModel->AddProperties(Path, Properties);
+			}
+			else
+			{
+				EditablePropertiesModel->RemoveProperties(Path, Properties);
+				const bool bNeedsToRemoveNonRoot = !Object->IsA<AActor>() && EditablePropertiesModel->GetNumProperties(Path) == 0;
+				if (bNeedsToRemoveNonRoot)
+				{
+					EditablePropertiesModel->RemoveObjects({ Path });
+				}
+			}
+		}
+	}
+
+	void SObjectToPropertyEditor::OnDeleteObjects(const TArray<TSharedPtr<FReplicatedObjectData>>& ObjectsToDelete) const
+	{
+		TArray<FSoftObjectPath> DeleteObjectPaths;
+		Algo::Transform(ObjectsToDelete, DeleteObjectPaths, [](const TSharedPtr<FReplicatedObjectData>& Data) { return Data->GetObjectPath(); });
+
+		// We want to delete children not listed in the outliner, such as components and other subobjects
+		TArray<FSoftObjectPath> ObjectAndChildren;
+		EditablePropertiesModel->ForEachReplicatedObject([&ObjectsToDelete, &ObjectAndChildren](const FSoftObjectPath& ReplicatedObject)
+		{
+			const bool bIsChildOfDeletedObject = Algo::AnyOf(ObjectsToDelete, [&ReplicatedObject](const TSharedPtr<FReplicatedObjectData>& ObjectData)
+			{
+				// ReplicatedObject is a child of a deleted object if the path before it contains one of the deleted objects
+				return ReplicatedObject.ToString().Contains(ObjectData->GetObjectPath().ToString());
+			});
+			if (bIsChildOfDeletedObject)
+			{
+				ObjectAndChildren.Add(ReplicatedObject);
+			}
+			return EBreakBehavior::Continue;
+		});
+		
+		checkf(ObjectAndChildren.Num() >= ObjectsToDelete.Num(), TEXT("Above algorithm is broken."));
+		EditablePropertiesModel->RemoveObjects(ObjectAndChildren);
+	}
+
+	TSharedPtr<SWidget> SObjectToPropertyEditor::OnObjectsContextMenuOpening() const
+	{
+		FMenuBuilder MenuBuilder(true, nullptr);
+		AddObjectSourceContextMenuOptions(MenuBuilder);
+		MenuBuilder.AddMenuEntry(
+				LOCTEXT("DeleteItems", "Delete"),
+				FText::GetEmpty(),
+				FSlateIcon(),
+				FUIAction(FExecuteAction::CreateSP(this, &SObjectToPropertyEditor::OnDeleteObjects_PassByValue, GetSelectedOutlinerObjects())),
+				NAME_None,
+				EUserInterfaceActionType::Button
+			);
+		
+		return MenuBuilder.MakeWidget();
+	}
+
+	void SObjectToPropertyEditor::AddObjectSourceContextMenuOptions(FMenuBuilder& MenuBuilder) const
+	{
+		using namespace ConcertSharedSlate;
+		
+		// Context menu generation is only supported for single items
+		const TArray<TSharedPtr<FReplicatedObjectData>> SelectedObjects = GetSelectedOutlinerObjects();
+		if (SelectedObjects.Num() != 1)
+		{
+			return;
+		}
+		
+		const TArray<TSharedRef<IObjectSourceModel>> ContextMenuOptions = ObjectSelectionSource->GetContextMenuOptions(SelectedObjects[0]->GetObjectPath());
+		if (ContextMenuOptions.IsEmpty())
+		{
+			return;
+		}
+		
+		const FSourceModelBuilders<FSelectableObjectInfo>::FItemPickerArgs Args = MakeObjectSourceBuilderArgs();
+		for (const TSharedRef<IObjectSourceModel>& SourceModel : ContextMenuOptions)
+		{
+			bool bCanAddAnyObject = false; 
+			SourceModel->EnumerateSelectableItems([this, &bCanAddAnyObject](const FSelectableObjectInfo& SelectableOption) mutable
+			{
+				bCanAddAnyObject |= !EditablePropertiesModel->ContainsObjects({ SelectableOption.Object.Get() });
+				return bCanAddAnyObject ? EBreakBehavior::Break : EBreakBehavior::Continue;
+			});
+
+			// Skip showing context menu options which will not add anything new.
+			if (bCanAddAnyObject)
+			{
+				FSourceModelBuilders<FSelectableObjectInfo>::AddOptionToMenu(SourceModel, Args, MenuBuilder);
+			}
+		}
+		
+		MenuBuilder.AddSeparator();
+	}
+
+	ConcertSharedSlate::FSourceModelBuilders<FSelectableObjectInfo>::FItemPickerArgs SObjectToPropertyEditor::MakeObjectSourceBuilderArgs() const
+	{
+		using namespace ConcertSharedSlate;
+		
+		return FSourceModelBuilders<FSelectableObjectInfo>::FItemPickerArgs
+		{
+			FSourceModelBuilders<FSelectableObjectInfo>::FOnItemsSelected::CreateSP(this, &SObjectToPropertyEditor::OnObjectsSelectedForAdding),
+			FSourceModelBuilders<FSelectableObjectInfo>::FGetItemDisplayString::CreateLambda([](const FSelectableObjectInfo& Item)
+			{
+				return Item.Object.IsValid() ? DisplayUtils::GetObjectDisplayString(*Item.Object.Get()) : TEXT("");
+			}),
+			FSourceModelBuilders<FSelectableObjectInfo>::FGetItemIcon::CreateLambda([](const FSelectableObjectInfo& Item)
+			{
+				return Item.Object.IsValid() ? DisplayUtils::GetObjectIcon(*Item.Object.Get()) : FSlateIcon();
+			}),
+			FSourceModelBuilders<FSelectableObjectInfo>::FIsItemSelected::CreateLambda([this](const FSelectableObjectInfo& Item)
+			{
+				return EditablePropertiesModel->ContainsObjects({ Item.Object.Get() } );
+			})
+		}; 
+	}
+
+	bool SObjectToPropertyEditor::SortBySelectionThenByName_PropertyPredicate(const TSharedPtr<FReplicatedPropertyData>& Left, const TSharedPtr<FReplicatedPropertyData>& Right) const
+	{
+		const TArray<FSoftObjectPath> SelectedObjects = GetSelectedObjectShowingProperties();
+		const ECheckBoxState LeftCheckboxState = ReplicationPropertyColumns::GetPropertyCheckboxStateBasedOnSelection(Left->GetProperty(), SelectedObjects, *EditablePropertiesModel);
+		const ECheckBoxState RightCheckboxState = ReplicationPropertyColumns::GetPropertyCheckboxStateBasedOnSelection(Right->GetProperty(), SelectedObjects, *EditablePropertiesModel);
+		
+		// Secondary sort by name
+		if (LeftCheckboxState == RightCheckboxState)
+		{
+			return DisplayUtils::GetPropertyDisplayString(Left->GetProperty()) < DisplayUtils::GetPropertyDisplayString(Right->GetProperty());
+		}
+
+		// Selected properties should appear first
+		return LeftCheckboxState == ECheckBoxState::Checked
+			&& (RightCheckboxState == ECheckBoxState::Unchecked || RightCheckboxState == ECheckBoxState::Undetermined);
+	}
+}
+
+#undef LOCTEXT_NAMESPACE
\ No newline at end of file
diff --git a/Engine/Plugins/Developer/Concert/ConcertUI/ConcertClientSharedSlate/Source/ConcertClientSharedSlate/Private/Replication/Editor/View/ObjectEditor/SObjectToPropertyEditor.h b/Engine/Plugins/Developer/Concert/ConcertUI/ConcertClientSharedSlate/Source/ConcertClientSharedSlate/Private/Replication/Editor/View/ObjectEditor/SObjectToPropertyEditor.h
new file mode 100644
--- /dev/null
+++ b/Engine/Plugins/Developer/Concert/ConcertUI/ConcertClientSharedSlate/Source/ConcertClientSharedSlate/Private/Replication/Editor/View/ObjectEditor/SObjectToPropertyEditor.h
@@ -0,0 +1,94 @@
+// Copyright Epic Games, Inc. All Rights Reserved.
+
+#pragma once
+
+#include "Model/Item/SourceModelBuilders.h"
+#include "Replication/Editor/Model/Property/IPropertySourceModel.h"
+#include "Replication/Editor/View/ObjectViewer/SObjectToPropertyView.h"
+#include "Widgets/DeclarativeSyntaxSupport.h"
+
+namespace UE::ConcertClientSharedSlate
+{
+	class IReplicationSubobjectView;
+}
+
+class SHorizontalBox;
+
+namespace UE::ConcertClientSharedSlate
+{
+	class IEditableObjectToPropertiesModel;
+	class FFakeObjectToPropertiesEditorModel;
+	class IObjectSelectionSourceModel;
+	class IObjectToPropertiesModel;
+	class IPropertySelectionSourceModel;
+	class IReplicationSubobjectView;
+	
+	enum class EReplicatedObjectChangeReason : uint8;
+	
+	struct FSelectableObjectInfo;
+	
+	/**
+	 * Root widget for editing UMultiUserPropertyReplicationSelection.
+	 * This widget extends SObjectToPropertyViewer by using IEditableObjectToPropertiesModel for editing.
+	 *
+	 * This widget assumes that the IReplicationSubobjectView::GetSelectedObjects can report subobjects that are
+	 * not in the model yet. Such objects would be added to the model.
+	 */
+	class CONCERTCLIENTSHAREDSLATE_API SObjectToPropertyEditor : public SObjectToPropertyView
+	{
+	public:
+		
+		SLATE_BEGIN_ARGS(SObjectToPropertyEditor)
+		{}
+			/** Optional. Placed between root object outliner and property editor. */
+			SLATE_ARGUMENT(TSharedPtr<IReplicationSubobjectView>, SubobjectView)
+		SLATE_END_ARGS()
+
+		void Construct(const FArguments& InArgs,
+           TSharedRef<IEditableObjectToPropertiesModel> InPropertiesModel,
+           TSharedRef<IObjectSelectionSourceModel> InObjectSelectionSource,
+           TSharedRef<IPropertySelectionSourceModel> InPropertySelectionSource
+		);
+		
+	private:
+
+		/** For reading and writting to the edited asset */
+		TSharedPtr<IEditableObjectToPropertiesModel> EditablePropertiesModel;
+		/**
+		 * Fakes to SObjectToPropertyViewer that all UClass properties are contained.
+		 * We inject checkboxes to SObjectToPropertyViewer which do the actual adding and removing.
+		 */
+		TSharedPtr<FFakeObjectToPropertiesEditorModel> PropertiesModelAdapter;
+		
+		/** For deciding which objects can be added to EditablePropertiesModel. */
+		TSharedPtr<IObjectSelectionSourceModel> ObjectSelectionSource;
+		/** For deciding which properties can be added to EditablePropertiesModel. */
+		TSharedPtr<IPropertySelectionSourceModel> PropertySelectionSource;
+
+		/** Contains the widgets generated for PropertySelectionSource (since they depend on the selected objects). */
+		TSharedPtr<SHorizontalBox> AddPropertyWidgetContainer;
+		
+		void OnObjectsChanged(TConstArrayView<UObject*> AddedObjects, TConstArrayView<FSoftObjectPath> RemovedObjects, EReplicatedObjectChangeReason ChangeReason);
+		void OnPropertiesChanged();
+
+		// Customizing left objects search bar
+		TSharedRef<SWidget> BuildRootAddObjectWidgets() const;
+		void OnObjectsSelectedForAdding(TArray<FSelectableObjectInfo> ObjectsToAdd) const;
+		
+		// Replicates property column events
+		ECheckBoxState OnGetPropertyCheckboxState(const FConcertPropertyChain& PropertyChain);
+		void OnPropertyCheckboxChanged(bool bIsChecked, const FConcertPropertyChain& PropertyChain);
+
+		// Respond to object editing events
+		void OnDeleteObjects_PassByValue(TArray<TSharedPtr<FReplicatedObjectData>> CopiedObjectsToDelete) const { OnDeleteObjects(CopiedObjectsToDelete); }
+		void OnDeleteObjects(const TArray<TSharedPtr<FReplicatedObjectData>>& ObjectsToDelete) const;
+		TSharedPtr<SWidget> OnObjectsContextMenuOpening() const;
+		void AddObjectSourceContextMenuOptions(FMenuBuilder& MenuBuilder) const;
+		
+		// Utils for building item source widgets
+		ConcertSharedSlate::FSourceModelBuilders<FSelectableObjectInfo>::FItemPickerArgs MakeObjectSourceBuilderArgs() const;
+
+		/** Util for sorting property data by*/
+		bool SortBySelectionThenByName_PropertyPredicate(const TSharedPtr<FReplicatedPropertyData>& Left, const TSharedPtr<FReplicatedPropertyData>& Right) const;
+	};
+}
diff --git a/Engine/Plugins/Developer/Concert/ConcertUI/ConcertClientSharedSlate/Source/ConcertClientSharedSlate/Private/Replication/Editor/View/ObjectEditor/SPropertyReplicationSelectionEditor.cpp b/Engine/Plugins/Developer/Concert/ConcertUI/ConcertClientSharedSlate/Source/ConcertClientSharedSlate/Private/Replication/Editor/View/ObjectEditor/SPropertyReplicationSelectionEditor.cpp
deleted file mode 100644
diff --git a/Engine/Plugins/Developer/Concert/ConcertUI/ConcertClientSharedSlate/Source/ConcertClientSharedSlate/Private/Replication/Editor/View/ObjectEditor/SPropertyReplicationSelectionEditor.h b/Engine/Plugins/Developer/Concert/ConcertUI/ConcertClientSharedSlate/Source/ConcertClientSharedSlate/Private/Replication/Editor/View/ObjectEditor/SPropertyReplicationSelectionEditor.h
deleted file mode 100644
diff --git a/Engine/Plugins/Developer/Concert/ConcertUI/ConcertClientSharedSlate/Source/ConcertClientSharedSlate/Private/Replication/Editor/View/ObjectViewer/Property/SReplicatedPropertiesView.h b/Engine/Plugins/Developer/Concert/ConcertUI/ConcertClientSharedSlate/Source/ConcertClientSharedSlate/Private/Replication/Editor/View/ObjectViewer/Property/SReplicatedPropertiesView.h
--- a/Engine/Plugins/Developer/Concert/ConcertUI/ConcertClientSharedSlate/Source/ConcertClientSharedSlate/Private/Replication/Editor/View/ObjectViewer/Property/SReplicatedPropertiesView.h
+++ b/Engine/Plugins/Developer/Concert/ConcertUI/ConcertClientSharedSlate/Source/ConcertClientSharedSlate/Private/Replication/Editor/View/ObjectViewer/Property/SReplicatedPropertiesView.h
@@ -5,7 +5,7 @@
 #include "Filters/SBasicFilterBar.h"
 #include "Widgets/DeclarativeSyntaxSupport.h"
 #include "Widgets/SCompoundWidget.h"
-#include "Replication/Editor/View/ObjectViewer/SReplicationTreeView.h"
+#include "Replication/Editor/View/ObjectViewer/Tree/SReplicationTreeView.h"
 
 namespace UE::ConcertClientSharedSlate
 {
diff --git a/Engine/Plugins/Developer/Concert/ConcertUI/ConcertClientSharedSlate/Source/ConcertClientSharedSlate/Private/Replication/Editor/View/ObjectViewer/ReplicationColumn.h b/Engine/Plugins/Developer/Concert/ConcertUI/ConcertClientSharedSlate/Source/ConcertClientSharedSlate/Private/Replication/Editor/View/ObjectViewer/ReplicationColumn.h
deleted file mode 100644
diff --git a/Engine/Plugins/Developer/Concert/ConcertUI/ConcertClientSharedSlate/Source/ConcertClientSharedSlate/Private/Replication/Editor/View/ObjectViewer/SObjectToPropertyView.cpp b/Engine/Plugins/Developer/Concert/ConcertUI/ConcertClientSharedSlate/Source/ConcertClientSharedSlate/Private/Replication/Editor/View/ObjectViewer/SObjectToPropertyView.cpp
new file mode 100644
--- /dev/null
+++ b/Engine/Plugins/Developer/Concert/ConcertUI/ConcertClientSharedSlate/Source/ConcertClientSharedSlate/Private/Replication/Editor/View/ObjectViewer/SObjectToPropertyView.cpp
@@ -0,0 +1,220 @@
+// Copyright Epic Games, Inc. All Rights Reserved.
+
+#include "SObjectToPropertyView.h"
+
+#include "ConcertFrontendUtils.h"
+#include "Replication/Editor/Model/IObjectToPropertiesModel.h"
+#include "Replication/Editor/Model/ReplicatedObjectData.h"
+#include "Replication/Editor/Model/ReplicatedPropertyData.h"
+#include "Replication/Editor/View/ObjectViewer/Property/SReplicatedPropertiesView.h"
+#include "Replication/Editor/View/ObjectViewer/Tree/SelectionViewerColumns.h"
+#include "SSubobjectAndPropertySection.h"
+
+#include "Widgets/Layout/SBorder.h"
+#include "Widgets/Layout/SSplitter.h"
+#include "Widgets/SBoxPanel.h"
+#include "Widgets/Text/STextBlock.h"
+
+#define LOCTEXT_NAMESPACE "SObjectToPropertyViewer"
+
+namespace UE::ConcertClientSharedSlate
+{
+	void SObjectToPropertyView::Construct(const FArguments& InArgs, TSharedRef<IObjectToPropertiesModel> InPropertiesModel)
+	{
+		PropertiesModel = MoveTemp(InPropertiesModel);
+		
+		ChildSlot
+		[
+			CreateContentWidget(InArgs)
+		];
+
+		RefreshObjectData();
+		RefreshSubobjectData();
+		RefreshPropertyData();
+
+		PropertyArea->SetExpanded(true);
+	}
+
+	void SObjectToPropertyView::RefreshObjectData()
+	{
+		const int32 NumElements = PropertiesModel->GetNumReplicatedObjects();
+		// Re-using existing instances is tricky: we cannot update the object path in an item because the list view will no detect this change;
+		// list view only looks at the shared ptr address. So the UI will not be refreshed. Since the number of items will be small, just reallocate... 
+		ObjectRowData.Empty(NumElements);
+
+		// Try to re-use old instances by using the old PathToObjectDataCache. This is also done so the expansion states restore correctly in the tree view.
+		TMap<FSoftObjectPath, TSharedPtr<FReplicatedObjectData>> NewPathToObjectDataCache;
+		
+		// Do a complete refresh.
+		// Complete refresh is acceptable because the list is updated infrequently and typically small < 500 items.
+		// An alternative would be to change RefreshObjectData to be called with two variables ObjectsAdded and ObjectsRemoved.
+		PropertiesModel->ForEachReplicatedObject([this, &NewPathToObjectDataCache](const FSoftObjectPath& Path) mutable
+		{
+			const TSharedPtr<FReplicatedObjectData>* ExistingItem = PathToObjectDataCache.Find(Path);
+			const TSharedRef<FReplicatedObjectData> Item = ExistingItem ? ExistingItem->ToSharedRef() : AllocateObjectData(Path);
+			ObjectRowData.Emplace(Item);
+			NewPathToObjectDataCache.Emplace(Path, Item);
+			return EBreakBehavior::Continue;
+		});
+
+		// Only refresh the tree if it is necessary as it causes us to select stuff in the subobject view
+		if (!PathToObjectDataCache.OrderIndependentCompareEqual(NewPathToObjectDataCache))
+		{
+			// If an item was removed, then NewPathToObjectDataCache does not contain it. 
+			PathToObjectDataCache = MoveTemp(NewPathToObjectDataCache);
+
+			// The tree view requires the item source to only contain the root items. Children are discovered via GetObjectRowChildren. We re-use GetObjectRowChildren to remove any non-root nodes.
+			BuildRootObjectRowData();
+			ReplicatedObjects->OnItemsChanged();
+		}
+	}
+
+	void SObjectToPropertyView::RefreshSubobjectData()
+	{
+		SubobjectAndPropertySection->RefreshSubobjectData();
+		
+		// Need to update properties in case the subobject view changed selection
+		RefreshPropertyData();
+	}
+
+	void SObjectToPropertyView::RefreshPropertyData()
+	{
+		SubobjectAndPropertySection->RefreshPropertyData();
+	}
+	
+	const TArray<TSharedPtr<FReplicatedPropertyData>>& SObjectToPropertyView::GetPropertyRowData() const
+	{
+		return SubobjectAndPropertySection->GetPropertyRowData();
+	}
+
+	TArray<FSoftObjectPath> SObjectToPropertyView::GetSelectedObjectShowingProperties() const
+	{
+		return SubobjectAndPropertySection->GetSelectedObjects();
+	}
+
+	TSharedRef<FReplicatedObjectData> SObjectToPropertyView::AllocateObjectData(FSoftObjectPath ObjectPath)
+	{
+		return MakeShared<FReplicatedObjectData>(MoveTemp(ObjectPath));
+	}
+
+	TSharedRef<SWidget> SObjectToPropertyView::CreateContentWidget(const FArguments& InArgs)
+	{
+		return SNew(SSplitter)
+			.Orientation(Orient_Vertical)
+
+			+SSplitter::Slot()
+			 .Value(1.f)
+			[
+				CreateActorsSection(InArgs)
+			]
+
+			+SSplitter::Slot()
+			.SizeRule(TAttribute<SSplitter::ESizeRule>(this, &SObjectToPropertyView::GetPropertyAreaSizeRule))
+			.Value(2.f)
+			[
+				CreatePropertiesSection(InArgs)
+			];
+	}
+
+	TSharedRef<SWidget> SObjectToPropertyView::CreateActorsSection(const FArguments& InArgs)
+	{
+		TArray<ReplicationObjectColumns::FReplicationObjectColumn> Columns
+		{
+			ReplicationObjectColumns::IconColumn(PropertiesModel.ToSharedRef()),
+			ReplicationObjectColumns::LabelColumn(),
+			ReplicationObjectColumns::TypeColumn(PropertiesModel.ToSharedRef())
+		};
+		Columns.Append(InArgs._AdditionalObjectColumns);
+		
+		return SAssignNew(ReplicatedObjects, SReplicationTreeView<TSharedPtr<FReplicatedObjectData>>)
+			.RootItemsSource(&RootObjectRowData)
+			.OnGetChildren(this, &SObjectToPropertyView::GetObjectRowChildren)
+			.OnContextMenuOpening(InArgs._OnObjectsContextMenuOpening)
+			.OnDeleteItems(InArgs._OnDeleteObjects)
+			.OnSelectionChanged_Lambda([this]()
+			{
+				RefreshSubobjectData();
+				// The subobject UI should automatically highlight the root objects after they are selected in the outliner
+				SubobjectAndPropertySection->SelectRootObjects();
+			})
+			.Columns(Columns)
+			.ExpandableColumnLabel(ReplicationObjectColumns::LabelColumnId)
+			.SelectionMode(ESelectionMode::Multi)
+			.LeftOfSearchBar()
+			[
+				InArgs._LeftOfObjectSearchBar.Widget
+			];
+	}
+
+	TSharedRef<SWidget> SObjectToPropertyView::CreatePropertiesSection(const FArguments& InArgs)
+	{
+		return SNew(SBorder)
+			.BorderImage(FAppStyle::Get().GetBrush("ToolPanel.GroupBorder"))
+			.Padding(0.0f)
+			[
+				SAssignNew(PropertyArea, SExpandableArea)
+				.InitiallyCollapsed(true) 
+				.BorderBackgroundColor(FLinearColor(0.6f, 0.6f, 0.6f, 1.0f))
+				.BorderImage_Lambda([this]() { return ConcertFrontendUtils::GetExpandableAreaBorderImage(*PropertyArea); })
+				.BodyBorderImage(FAppStyle::Get().GetBrush("ToolPanel.GroupBorder"))
+				.BodyBorderBackgroundColor(FLinearColor::White)
+				.OnAreaExpansionChanged(this, &SObjectToPropertyView::OnPropertyAreaExpansionChanged)
+				.Padding(0.0f)
+				.HeaderContent()
+				[
+					SNew(STextBlock)
+					.Text(LOCTEXT("ReplicatedProperties", "Subobjects & Properties"))
+					.Font(FAppStyle::Get().GetFontStyle("DetailsView.CategoryFontStyle"))
+					.ShadowOffset(FVector2D(1.0f, 1.0f))
+				]
+				.BodyContent()
+				[
+					SAssignNew(SubobjectAndPropertySection, SSubobjectAndPropertySection, PropertiesModel.ToSharedRef())
+					.AdditionalPropertyColumns(InArgs._AdditionalPropertyColumns)
+					.SubobjectView(InArgs._SubobjectView)
+					.SortPropertyRowPredicate(InArgs._SortPropertyRowPredicate)
+					.GetSelectedRootObjects_Lambda([this](){ return GetSelectedOutlinerObjects(); })
+					.LeftOfPropertySearchBar()
+					[
+						InArgs._LeftOfPropertySearchBar.Widget
+					]
+				]
+			];
+	}
+
+	void SObjectToPropertyView::BuildRootObjectRowData()
+	{
+		TSet<TSharedPtr<FReplicatedObjectData>> NonRootNodes;
+		for (const TSharedPtr<FReplicatedObjectData>& Node : ObjectRowData)
+		{
+			GetObjectRowChildren(Node, [&NonRootNodes](TSharedPtr<FReplicatedObjectData> Child)
+			{
+				NonRootNodes.Add(MoveTemp(Child));
+			});
+		}
+
+		// Make RootObjectRowData only contain those nodes which were not listed as children 
+		RootObjectRowData.Empty(NonRootNodes.Num());
+		for (const TSharedPtr<FReplicatedObjectData>& Node : ObjectRowData)
+		{
+			if (!NonRootNodes.Contains(Node))
+			{
+				RootObjectRowData.Add(Node);
+			}
+		}
+		
+		RootObjectRowData.Sort([](const TSharedPtr<FReplicatedObjectData>& Left, const TSharedPtr<FReplicatedObjectData>& Right)
+		{
+			return Left->GetObjectPath().GetSubPathString() < Right->GetObjectPath().GetSubPathString();
+		});
+	}
+
+	void SObjectToPropertyView::GetObjectRowChildren(TSharedPtr<FReplicatedObjectData> ReplicatedObjectData, TFunctionRef<void(TSharedPtr<FReplicatedObjectData>)> ProcessChild)
+	{
+		// For now there are no children in the outliner.
+		// In the future we could add an external delegate that allows child actors to be parented.
+		// Important: this view should be possible to be built in programs, so it should not reference things like AActor, UActorComponent, ResolveObject, etc. directly.
+	}
+}
+
+#undef LOCTEXT_NAMESPACE
\ No newline at end of file
diff --git a/Engine/Plugins/Developer/Concert/ConcertUI/ConcertClientSharedSlate/Source/ConcertClientSharedSlate/Private/Replication/Editor/View/ObjectViewer/SObjectToPropertyView.h b/Engine/Plugins/Developer/Concert/ConcertUI/ConcertClientSharedSlate/Source/ConcertClientSharedSlate/Private/Replication/Editor/View/ObjectViewer/SObjectToPropertyView.h
new file mode 100644
--- /dev/null
+++ b/Engine/Plugins/Developer/Concert/ConcertUI/ConcertClientSharedSlate/Source/ConcertClientSharedSlate/Private/Replication/Editor/View/ObjectViewer/SObjectToPropertyView.h
@@ -0,0 +1,126 @@
+// Copyright Epic Games, Inc. All Rights Reserved.
+
+#pragma once
+
+#include "Replication/Editor/Model/ReplicatedObjectData.h"
+#include "Replication/Editor/View/ObjectViewer/Tree/SReplicationTreeView.h"
+#include "Replication/Editor/View/ObjectViewer/Tree/SelectionViewerColumns.h"
+#include "SSubobjectAndPropertySection.h"
+
+#include "Algo/Transform.h"
+#include "Misc/Optional.h"
+#include "Widgets/DeclarativeSyntaxSupport.h"
+#include "Widgets/SCompoundWidget.h"
+
+class SExpandableArea;
+class SWidgetSwitcher;
+struct FSoftClassPath;
+struct FSoftObjectPath;
+
+namespace UE::ConcertClientSharedSlate
+{
+	class SSubobjectAndPropertySection;
+	class FReplicatedPropertyData;
+	class FReplicatedObjectData;
+	class IEditableObjectToPropertiesModel;
+	class IReplicationSubobjectView;
+	class IObjectToPropertiesModel;
+	class SReplicatedPropertiesView;
+	
+	/**
+	 * Root widget for viewing UMultiUserPropertyReplicationSelection.
+	 * This widget knows how to display IObjectToPropertiesModel.
+	 * 
+	 * The underlying data is modified by SObjectToPropertyEditor, which uses this widget's extension
+	 * points to call functions on IEditableObjectToPropertiesModel.
+	 *
+	 * Important: this view should be possible to be built in programs, so it should not reference things like AActor,
+	 * UActorComponent, ResolveObject, etc. directly. 
+	 */
+	class CONCERTCLIENTSHAREDSLATE_API SObjectToPropertyView : public SCompoundWidget
+	{
+	public:
+
+		SLATE_BEGIN_ARGS(SObjectToPropertyView)
+		{}
+			/** Additional columns to add to the object view */
+			SLATE_ARGUMENT(TArray<ReplicationObjectColumns::FReplicationObjectColumn>, AdditionalObjectColumns)
+			/** Additional columns to add to the property view */
+			SLATE_ARGUMENT(TArray<ReplicationPropertyColumns::FReplicationPropertyColumn>, AdditionalPropertyColumns)
+
+			/** Optional. Placed between root object outliner and property editor. */
+			SLATE_ARGUMENT(TSharedPtr<IReplicationSubobjectView>, SubobjectView)
+
+			/** Optional. Called when the delete key is pressed in the object view. */
+			SLATE_EVENT(SReplicationTreeView<TSharedPtr<FReplicatedObjectData>>::FDeleteItems, OnDeleteObjects)
+		
+			/** Called to generate the context menu for objects. */
+			SLATE_EVENT(FOnContextMenuOpening, OnObjectsContextMenuOpening)
+
+			/** Optional. Used for determining the order in which properties are displayed. */
+			SLATE_EVENT(FSortPropertyPredicate, SortPropertyRowPredicate)
+		
+			/** Optional widget to add to the left of the object list search bar. */
+			SLATE_NAMED_SLOT(FArguments, LeftOfObjectSearchBar)
+			/** Optional widget to add to the left of the property list search bar. */
+			SLATE_NAMED_SLOT(FArguments, LeftOfPropertySearchBar)
+		SLATE_END_ARGS()
+
+		void Construct(const FArguments& InArgs, TSharedRef<IObjectToPropertiesModel> InPropertiesModel);
+
+		void RefreshObjectData();
+		void RefreshSubobjectData();
+		void RefreshPropertyData();
+
+		const TSet<TSharedPtr<FReplicatedObjectData>>& GetObjectRowData() const { return ObjectRowData; }
+		const TArray<TSharedPtr<FReplicatedPropertyData>>& GetPropertyRowData() const;
+
+		/** @return Gets the root objects selected in the outliner; the subobject view chooses which of these objects (or their subobjects) end up in GetSelectedObjectShowingProperties. */
+		TArray<TSharedPtr<FReplicatedObjectData>> GetSelectedOutlinerObjects() const { return ReplicatedObjects->GetSelectedItems(); }
+		/** @return The objects for which properties are being shown; this is what the subobject reports are selected (or the root objects if there is no subobject view). */
+		TArray<FSoftObjectPath> GetSelectedObjectShowingProperties() const;
+		
+	private:
+
+		/** The model this view is visualizing. */
+		TSharedPtr<IObjectToPropertiesModel> PropertiesModel;
+
+		/** Lists the properties of the selected actor */
+		TSharedPtr<SExpandableArea> PropertyArea;
+		/** Edits the property list and (optionally) exposes subobjects of the selected root object. */
+		TSharedPtr<SSubobjectAndPropertySection> SubobjectAndPropertySection;
+
+		/** Tree view for replicated objects. */
+		TSharedPtr<SReplicationTreeView<TSharedPtr<FReplicatedObjectData>>> ReplicatedObjects;
+		
+		/**
+		 * These instances can be subclasses of FReplicatedObjectData, e.g. FReplicatedObjectData_Editor.
+		 * Their type can be overridden by subclasses.
+		 * They only have the FReplicatedObjectData type so they can be passed efficiently to SObjectToPropertyViewer.
+		 * @see GetObjectData
+		 */
+		TSet<TSharedPtr<FReplicatedObjectData>> ObjectRowData;
+		/** The instances of ObjectRowData which do not have any parents. This acts as the item source for the tree view. */
+		TArray<TSharedPtr<FReplicatedObjectData>> RootObjectRowData;
+		/** Inverse map of ObjectRowData using FReplicatedObjectData::GetObjectPath as key. Contains all elements of ObjectRowData. */
+		TMap<FSoftObjectPath, TSharedPtr<FReplicatedObjectData>> PathToObjectDataCache;
+
+		bool bIsPropertyAreaExpanded = false;
+
+		TSharedRef<FReplicatedObjectData> AllocateObjectData(FSoftObjectPath ObjectPath);
+
+		// Widget creation helpers
+		TSharedRef<SWidget> CreateContentWidget(const FArguments& InArgs);
+		TSharedRef<SWidget> CreateActorsSection(const FArguments& InArgs);
+		TSharedRef<SWidget> CreatePropertiesSection(const FArguments& InArgs);
+
+		/** Sets RootObjectRowData to all non-root nodes from ObjectRowData. */
+		void BuildRootObjectRowData();
+		
+		void GetObjectRowChildren(TSharedPtr<FReplicatedObjectData> ReplicatedObjectData, TFunctionRef<void(TSharedPtr<FReplicatedObjectData>)> ProcessChild);
+		
+		/** Handles how much space the 'Clients' area uses with respect to its expansion state. */
+		SSplitter::ESizeRule GetPropertyAreaSizeRule() const { return bIsPropertyAreaExpanded ? SSplitter::ESizeRule::FractionOfParent : SSplitter::ESizeRule::SizeToContent; }
+		void OnPropertyAreaExpansionChanged(bool bExpanded) { bIsPropertyAreaExpanded = bExpanded; }
+	};
+}
diff --git a/Engine/Plugins/Developer/Concert/ConcertUI/ConcertClientSharedSlate/Source/ConcertClientSharedSlate/Private/Replication/Editor/View/ObjectViewer/SPropertyReplicationSelectionViewer.cpp b/Engine/Plugins/Developer/Concert/ConcertUI/ConcertClientSharedSlate/Source/ConcertClientSharedSlate/Private/Replication/Editor/View/ObjectViewer/SPropertyReplicationSelectionViewer.cpp
deleted file mode 100644
diff --git a/Engine/Plugins/Developer/Concert/ConcertUI/ConcertClientSharedSlate/Source/ConcertClientSharedSlate/Private/Replication/Editor/View/ObjectViewer/SPropertyReplicationSelectionViewer.h b/Engine/Plugins/Developer/Concert/ConcertUI/ConcertClientSharedSlate/Source/ConcertClientSharedSlate/Private/Replication/Editor/View/ObjectViewer/SPropertyReplicationSelectionViewer.h
deleted file mode 100644
diff --git a/Engine/Plugins/Developer/Concert/ConcertUI/ConcertClientSharedSlate/Source/ConcertClientSharedSlate/Private/Replication/Editor/View/ObjectViewer/SReplicationColumnRow.h b/Engine/Plugins/Developer/Concert/ConcertUI/ConcertClientSharedSlate/Source/ConcertClientSharedSlate/Private/Replication/Editor/View/ObjectViewer/SReplicationColumnRow.h
deleted file mode 100644
diff --git a/Engine/Plugins/Developer/Concert/ConcertUI/ConcertClientSharedSlate/Source/ConcertClientSharedSlate/Private/Replication/Editor/View/ObjectViewer/SReplicationTreeView.h b/Engine/Plugins/Developer/Concert/ConcertUI/ConcertClientSharedSlate/Source/ConcertClientSharedSlate/Private/Replication/Editor/View/ObjectViewer/SReplicationTreeView.h
deleted file mode 100644
diff --git a/Engine/Plugins/Developer/Concert/ConcertUI/ConcertClientSharedSlate/Source/ConcertClientSharedSlate/Private/Replication/Editor/View/ObjectViewer/SSubobjectAndPropertySection.cpp b/Engine/Plugins/Developer/Concert/ConcertUI/ConcertClientSharedSlate/Source/ConcertClientSharedSlate/Private/Replication/Editor/View/ObjectViewer/SSubobjectAndPropertySection.cpp
new file mode 100644
--- /dev/null
+++ b/Engine/Plugins/Developer/Concert/ConcertUI/ConcertClientSharedSlate/Source/ConcertClientSharedSlate/Private/Replication/Editor/View/ObjectViewer/SSubobjectAndPropertySection.cpp
@@ -0,0 +1,263 @@
+// Copyright Epic Games, Inc. All Rights Reserved.
+
+#include "SSubobjectAndPropertySection.h"
+
+#include "Replication/Editor/Model/IObjectToPropertiesModel.h"
+#include "Replication/Editor/Model/ReplicatedObjectData.h"
+#include "Replication/Editor/Model/ReplicatedPropertyData.h"
+#include "Replication/Editor/View/IReplicationSubobjectView.h"
+#include "Replication/Editor/View/ObjectViewer/Property/SReplicatedPropertiesView.h"
+#include "Replication/Editor/View/ObjectViewer/Tree/SelectionViewerColumns.h"
+
+#include "Algo/AllOf.h"
+#include "Algo/ForEach.h"
+#include "Widgets/Layout/SBorder.h"
+#include "Widgets/Layout/SSplitter.h"
+#include "Widgets/Layout/SWidgetSwitcher.h"
+#include "Widgets/SBoxPanel.h"
+#include "Widgets/Text/STextBlock.h"
+
+#define LOCTEXT_NAMESPACE "SSubobjectAndPropertiesSection"
+
+namespace UE::ConcertClientSharedSlate
+{
+	void SSubobjectAndPropertySection::Construct(const FArguments& InArgs, TSharedRef<IObjectToPropertiesModel> InPropertiesModel)
+	{
+		PropertiesModel = MoveTemp(InPropertiesModel);
+		SortPropertyRowPredicate = InArgs._SortPropertyRowPredicate;
+		GetSelectedRootObjectsDelegate = InArgs._GetSelectedRootObjects;
+		check(GetSelectedRootObjectsDelegate.IsBound());
+		
+		ChildSlot
+		[
+			CreateSubobjectsAndPropertiesSection(InArgs)
+		];
+	}
+
+	void SSubobjectAndPropertySection::SelectRootObjects() const
+	{
+		if (SubobjectView)
+		{
+			SubobjectView->SelectRootObjects();
+		}
+	}
+
+	void SSubobjectAndPropertySection::RefreshSubobjectData()
+	{
+		if (!SubobjectView)
+		{
+			return;
+		}
+		
+		const TArray<TSharedPtr<FReplicatedObjectData>> SelectedObjects = GetSelectedRootObjectsDelegate.Execute();
+		TArray<FSoftObjectPath> SelectedObjectPaths;
+		Algo::Transform(SelectedObjects, SelectedObjectPaths, [](const TSharedPtr<FReplicatedObjectData>& Item){ return Item->GetObjectPath(); });
+		SubobjectView->SetRootObjects(SelectedObjectPaths);
+	}
+	
+	void SSubobjectAndPropertySection::RefreshPropertyData()
+	{
+		const TArray<FSoftObjectPath> SelectedObjects = GetSelectedObjects();
+		if (SelectedObjects.IsEmpty())
+		{
+			SetPropertyContent(EReplicatedPropertyContent::NoSelection);
+			return;
+		}
+
+		// Technically, the classes just need to be compatible with each other... but it is easier to just allow the same class.
+		const TOptional<FSoftClassPath> SharedClass = GetClassForPropertiesFromSelection(SelectedObjects);
+		if (!SharedClass)
+		{
+			SetPropertyContent(EReplicatedPropertyContent::SelectionTooBig);
+			return;
+		}
+		const FSoftClassPath Class = *SharedClass;
+		if (!ensureMsgf(Class.IsValid(), TEXT("This should only trigger if IReplicationSubobjectView returned a object not contained the model & that the UI is view-only. In that case fix your IReplicationSubobjectView or use it in an UI editor.")))
+		{
+			return;
+		}
+
+		// Build the set of properties that are shared by all of the selected objects
+		TSet<FConcertPropertyChain> SharedProperties = PropertiesModel->GetAllProperties(SelectedObjects[0]);
+		for (int32 i = 1; i < SelectedObjects.Num(); ++i)
+		{
+			SharedProperties = PropertiesModel->GetAllProperties(SelectedObjects[i]).Union(SharedProperties);
+		}
+		
+		// Try to re-use old instances by using the old ChainToPropertyDataCache. This is also done so the expansion states restore correctly in the tree view.
+		TMap<FConcertPropertyChain, TSharedPtr<FReplicatedPropertyData>> NewChainToPropertyDataCache;
+		
+		PropertyRowData.Empty(PropertiesModel->GetNumProperties(SelectedObjects[0]));
+		for (const FConcertPropertyChain& PropertyChain : SharedProperties)
+		{
+			const TSharedPtr<FReplicatedPropertyData>* ExistingItem = ChainToPropertyDataCache.Find(PropertyChain);
+			const TSharedRef<FReplicatedPropertyData> Item = ExistingItem ? ExistingItem->ToSharedRef() : AllocatePropertyData(Class, PropertyChain);
+			PropertyRowData.Emplace(Item);
+			NewChainToPropertyDataCache.Emplace(PropertyChain, Item);
+		}
+
+		// If an item was removed, then NewPathToPropertyDataCache does not contain it. 
+		ChainToPropertyDataCache = MoveTemp(NewChainToPropertyDataCache);
+		
+		// The tree view requires the item source to only contain the root items.
+		BuildRootPropertyRowData();
+		SetPropertyContent(EReplicatedPropertyContent::Properties);
+		ReplicatedProperties->OnItemsChanged();
+	}
+
+	TArray<FSoftObjectPath> SSubobjectAndPropertySection::GetSelectedObjects() const
+	{
+		if (SubobjectView)
+		{
+			return SubobjectView->GetSelectedObjects();
+		}
+
+		TArray<FSoftObjectPath> Result;
+		Algo::Transform(GetSelectedRootObjectsDelegate.Execute(), Result, [](const TSharedPtr<FReplicatedObjectData>& ObjectData)
+		{
+			return ObjectData->GetObjectPath();
+		});
+		return Result;
+	}
+
+	TSharedRef<SWidget> SSubobjectAndPropertySection::CreateSubobjectsAndPropertiesSection(const FArguments& InArgs)
+	{
+		if (InArgs._SubobjectView.IsValid())
+		{
+			SubobjectView = InArgs._SubobjectView;
+			SubobjectView->OnSelectionChanged().AddSP(this, &SSubobjectAndPropertySection::OnSubobjectSelectionChanged);
+			return SNew(SSplitter)
+				.Orientation(Orient_Vertical)
+
+				+SSplitter::Slot()
+				.Value(1.f)
+				[
+					SubobjectView.ToSharedRef()
+				]
+			
+				+SSplitter::Slot()
+				.Value(2.f)
+				[
+                	CreatePropertiesView(InArgs)
+                ];
+		}
+		
+		return CreatePropertiesView(InArgs);
+	}
+
+	TSharedRef<SWidget> SSubobjectAndPropertySection::CreatePropertiesView(const FArguments& InArgs)
+	{
+		TArray<ReplicationPropertyColumns::FReplicationPropertyColumn> Columns
+		{
+			ReplicationPropertyColumns::LabelColumn(),
+			ReplicationPropertyColumns::TypeColumn()
+		};
+		Columns.Append(InArgs._AdditionalPropertyColumns);
+		
+		return SAssignNew(PropertyContent, SWidgetSwitcher)
+			// Make sure the slots are coherent with the order of EReplicatedPropertyContent!
+			.WidgetIndex(static_cast<int32>(EReplicatedPropertyContent::NoSelection))
+			
+			// EReplicatedPropertyContent::Properties
+			+SWidgetSwitcher::Slot()
+			[
+				SAssignNew(ReplicatedProperties, SReplicatedPropertiesView)
+				.RootItemsSource(&RootPropertyRowData)
+				.OnGetChildren(this, &SSubobjectAndPropertySection::GetPropertyRowChildren)
+				.Columns(Columns)
+				.ExpandableColumnLabel(ReplicationPropertyColumns::LabelColumnId)
+				.SelectionMode(ESelectionMode::Multi)
+				.LeftOfSearchBar()
+				[
+					InArgs._LeftOfPropertySearchBar.Widget
+				]
+			]
+			
+			// EReplicatedPropertyContent::NoSelection
+			+SWidgetSwitcher::Slot()
+			.VAlign(VAlign_Center)
+			.HAlign(HAlign_Center)
+			[
+				SNew(STextBlock)
+				.Text(LOCTEXT("NoSelection", "Select an object to see selected properties"))
+			]
+			
+			// EReplicatedPropertyContent::SelectionTooBig
+			+SWidgetSwitcher::Slot()
+			.VAlign(VAlign_Center)
+			.HAlign(HAlign_Center)
+			[
+				SNew(STextBlock)
+				.Text(LOCTEXT("SelectionTooBig", "Select objects of the same type type to see selected properties"))
+			];
+	}
+	
+	TSharedRef<FReplicatedPropertyData> SSubobjectAndPropertySection::AllocatePropertyData(FSoftClassPath OwningClass, FConcertPropertyChain PropertyChain)
+	{
+		return MakeShared<FReplicatedPropertyData>(MoveTemp(OwningClass), MoveTemp(PropertyChain));
+	}
+
+	void SSubobjectAndPropertySection::BuildRootPropertyRowData()
+	{
+		RootPropertyRowData.Empty(PropertyRowData.Num());
+		for (const TSharedPtr<FReplicatedPropertyData>& PropertyData : PropertyRowData)
+		{
+			if (PropertyData->GetProperty().IsRootProperty())
+			{
+				RootPropertyRowData.Emplace(PropertyData);
+			}
+		}
+
+		SortPropertyRowArray(RootPropertyRowData);
+	}
+	
+	TOptional<FSoftClassPath> SSubobjectAndPropertySection::GetClassForPropertiesFromSelection(const TArray<FSoftObjectPath>& Objects) const
+	{
+		const FSoftClassPath Class = PropertiesModel->GetObjectClass(Objects[0]);
+		const bool bAllHaveSameClass = Algo::AllOf(Objects, [this, Class](const FSoftObjectPath& Object)
+		{
+			return PropertiesModel->GetObjectClass(Object) == Class;
+		});
+		return bAllHaveSameClass ? Class : TOptional<FSoftClassPath>{};
+	}
+
+	void SSubobjectAndPropertySection::GetPropertyRowChildren(TSharedPtr<FReplicatedPropertyData> ReplicatedPropertyData, TFunctionRef<void(TSharedPtr<FReplicatedPropertyData>)> ProcessChild)
+	{
+		TArray<TSharedPtr<FReplicatedPropertyData>> Children;
+		
+		// Not the most efficient but it should be fine.
+		for (const TSharedPtr<FReplicatedPropertyData>& Data : PropertyRowData)
+		{
+			if (Data->GetProperty().IsDirectChildOf(ReplicatedPropertyData->GetProperty()))
+			{
+				Children.Add(Data);
+			}
+		}
+
+		SortPropertyRowArray(Children);
+		Algo::ForEach(Children, [&ProcessChild](const TSharedPtr<FReplicatedPropertyData>& Data){ ProcessChild(Data); });
+	}
+
+	void SSubobjectAndPropertySection::OnSubobjectSelectionChanged()
+	{
+		RefreshPropertyData();
+	}
+
+	void SSubobjectAndPropertySection::SortPropertyRowArray(TArray<TSharedPtr<FReplicatedPropertyData>>& ToSort) const
+	{
+		if (SortPropertyRowPredicate.IsBound())
+		{
+			ToSort.Sort([this](const TSharedPtr<FReplicatedPropertyData>& Left, const TSharedPtr<FReplicatedPropertyData>& Right)
+			{
+				return SortPropertyRowPredicate.Execute(Left, Right);
+			});
+		}
+	}
+
+	void SSubobjectAndPropertySection::SetPropertyContent(EReplicatedPropertyContent Content) const
+	{
+		PropertyContent->SetActiveWidgetIndex(static_cast<int32>(Content));
+	}
+}
+
+#undef LOCTEXT_NAMESPACE
\ No newline at end of file
diff --git a/Engine/Plugins/Developer/Concert/ConcertUI/ConcertClientSharedSlate/Source/ConcertClientSharedSlate/Private/Replication/Editor/View/ObjectViewer/SSubobjectAndPropertySection.h b/Engine/Plugins/Developer/Concert/ConcertUI/ConcertClientSharedSlate/Source/ConcertClientSharedSlate/Private/Replication/Editor/View/ObjectViewer/SSubobjectAndPropertySection.h
new file mode 100644
--- /dev/null
+++ b/Engine/Plugins/Developer/Concert/ConcertUI/ConcertClientSharedSlate/Source/ConcertClientSharedSlate/Private/Replication/Editor/View/ObjectViewer/SSubobjectAndPropertySection.h
@@ -0,0 +1,118 @@
+// Copyright Epic Games, Inc. All Rights Reserved.
+
+#pragma once
+
+#include "Replication/Editor/View/ObjectViewer/Tree/SelectionViewerColumns.h"
+
+#include "Misc/Optional.h"
+#include "Replication/Data/ConcertPropertySelection.h"
+#include "Widgets/DeclarativeSyntaxSupport.h"
+#include "Widgets/SCompoundWidget.h"
+
+class SWidgetSwitcher;
+
+namespace UE::ConcertClientSharedSlate
+{
+	class SReplicatedPropertiesView;
+	class IReplicationSubobjectView;
+
+	DECLARE_DELEGATE_RetVal_TwoParams(bool, FSortPropertyPredicate, const TSharedPtr<FReplicatedPropertyData>& Left, const TSharedPtr<FReplicatedPropertyData>& Right);
+	
+	/**
+	 * Displays an optional IReplicationSubobjectView and followed by SReplicatedPropertiesView tree view (which is always created).
+	 * 
+	 * Important: this view should be possible to be built in programs, so it should not reference things like AActor,
+	 * UActorComponent, ResolveObject, etc. directly. 
+	 */
+	class SSubobjectAndPropertySection : public SCompoundWidget
+	{
+	public:
+		
+		DECLARE_DELEGATE_RetVal(TArray<TSharedPtr<FReplicatedObjectData>>, FGetSelectedRootObjects)
+
+		SLATE_BEGIN_ARGS(SSubobjectAndPropertySection)
+		{}
+			/** Additional columns to add to the property view */
+			SLATE_ARGUMENT(TArray<ReplicationPropertyColumns::FReplicationPropertyColumn>, AdditionalPropertyColumns)
+			/** Optional. Placed between root object outliner and property editor. */
+			SLATE_ARGUMENT(TSharedPtr<IReplicationSubobjectView>, SubobjectView)
+		
+			/** Optional. Used for determining the order in which properties are displayed. */
+			SLATE_EVENT(FSortPropertyPredicate, SortPropertyRowPredicate)
+			/** Gets the root objects selected in the object outliner. */
+			SLATE_EVENT(FGetSelectedRootObjects, GetSelectedRootObjects)
+		
+			/** Optional widget to add to the left of the property list search bar. */
+			SLATE_NAMED_SLOT(FArguments, LeftOfPropertySearchBar)
+		SLATE_END_ARGS()
+
+		void Construct(const FArguments& InArgs, TSharedRef<IObjectToPropertiesModel> InPropertiesModel);
+
+		/** If there is a subobject view, makes it select the root objects. */
+		void SelectRootObjects() const;
+		
+		void RefreshSubobjectData();
+		void RefreshPropertyData();
+		
+		const TArray<TSharedPtr<FReplicatedPropertyData>>& GetPropertyRowData() const { return PropertyRowData; }
+		TArray<FSoftObjectPath> GetSelectedObjects() const;
+
+	private:
+		
+		/** The model this view is visualizing. */
+		TSharedPtr<IObjectToPropertiesModel> PropertiesModel;
+		
+		/** Optional. External widget that selects subobjects from ActorArea. */
+		TSharedPtr<IReplicationSubobjectView> SubobjectView;
+		/** Tree view for replicated properties. Content depends on the current object selected. */
+		TSharedPtr<SReplicatedPropertiesView> ReplicatedProperties;
+		
+		enum class EReplicatedPropertyContent
+		{
+			/** Shows the properties */
+			Properties,
+			/** Prompts: "Select an object to see selected properties" */
+			NoSelection,
+			/** Prompts: "Select objects of the same type type to see selected properties" */
+			SelectionTooBig
+		};
+		/** Determines the content displayed for PropertyArea. */
+		TSharedPtr<SWidgetSwitcher> PropertyContent;
+		
+		/** Used for determining the order in which properties are displayed. Default: Sort by name. */
+		FSortPropertyPredicate SortPropertyRowPredicate;
+		/** Gets the root objects selected in the object outliner. */
+		FGetSelectedRootObjects GetSelectedRootObjectsDelegate;
+		
+		/**
+		 * These instances can be subclasses of FReplicatedPropertyData, e.g. FReplicatedPropertyData_Editor.
+		 * Their type can be overridden by subclasses.
+		 * They only have the FReplicatedPropertyData type so they can be passed efficiently to SObjectToPropertyViewer.
+		 * @see GetPropertyData
+		 */
+		TArray<TSharedPtr<FReplicatedPropertyData>> PropertyRowData;
+		/** The instances of ObjectRowData which do not have any parents. This acts as the item source for the tree view. */
+		TArray<TSharedPtr<FReplicatedPropertyData>> RootPropertyRowData;
+		/** Inverse map of PropertyRowData using FReplicatedPropertyData::GetProperty as key. Contains all elements of PropertyRowData. */
+		TMap<FConcertPropertyChain, TSharedPtr<FReplicatedPropertyData>> ChainToPropertyDataCache;
+
+		TSharedRef<SWidget> CreateSubobjectsAndPropertiesSection(const FArguments& InArgs);
+		TSharedRef<SWidget> CreatePropertiesView(const FArguments& InArgs);
+		
+		TSharedRef<FReplicatedPropertyData> AllocatePropertyData(FSoftClassPath OwningClass, FConcertPropertyChain PropertyChain);
+
+		/** Inits RootPropertyRowData from PropertyRowData. */
+		void BuildRootPropertyRowData();
+		/** Given the selected objects, determines whether they all have the same class and returns it if so. */
+		TOptional<FSoftClassPath> GetClassForPropertiesFromSelection(const TArray<FSoftObjectPath>& Objects) const;
+		void GetPropertyRowChildren(TSharedPtr<FReplicatedPropertyData> ReplicatedPropertyData, TFunctionRef<void(TSharedPtr<FReplicatedPropertyData>)> ProcessChild);
+		
+		void OnSubobjectSelectionChanged();
+		
+		// Utils
+		void SortPropertyRowArray(TArray<TSharedPtr<FReplicatedPropertyData>>& ToSort) const;
+		void SetPropertyContent(EReplicatedPropertyContent Content) const;
+	};
+}
+
+
diff --git a/Engine/Plugins/Developer/Concert/ConcertUI/ConcertClientSharedSlate/Source/ConcertClientSharedSlate/Private/Replication/Editor/View/ObjectViewer/SelectionViewerColumns.cpp b/Engine/Plugins/Developer/Concert/ConcertUI/ConcertClientSharedSlate/Source/ConcertClientSharedSlate/Private/Replication/Editor/View/ObjectViewer/SelectionViewerColumns.cpp
deleted file mode 100644
diff --git a/Engine/Plugins/Developer/Concert/ConcertUI/ConcertClientSharedSlate/Source/ConcertClientSharedSlate/Private/Replication/Editor/View/ObjectViewer/SelectionViewerColumns.h b/Engine/Plugins/Developer/Concert/ConcertUI/ConcertClientSharedSlate/Source/ConcertClientSharedSlate/Private/Replication/Editor/View/ObjectViewer/SelectionViewerColumns.h
deleted file mode 100644
diff --git a/Engine/Plugins/Developer/Concert/ConcertUI/ConcertClientSharedSlate/Source/ConcertClientSharedSlate/Private/Replication/Editor/View/ObjectViewer/Tree/ReplicationColumn.h b/Engine/Plugins/Developer/Concert/ConcertUI/ConcertClientSharedSlate/Source/ConcertClientSharedSlate/Private/Replication/Editor/View/ObjectViewer/Tree/ReplicationColumn.h
new file mode 100644
--- /dev/null
+++ b/Engine/Plugins/Developer/Concert/ConcertUI/ConcertClientSharedSlate/Source/ConcertClientSharedSlate/Private/Replication/Editor/View/ObjectViewer/Tree/ReplicationColumn.h
@@ -0,0 +1,72 @@
+// Copyright Epic Games, Inc. All Rights Reserved.
+
+#pragma once
+
+#include "Widgets/SNullWidget.h"
+#include "Widgets/Views/SHeaderRow.h"
+
+namespace UE::ConcertClientSharedSlate
+{
+	template<typename TListItemType>
+	class TReplicationColumn : public SHeaderRow::FColumn
+	{
+	public:
+
+		struct FBuildArgs
+		{
+			TSharedPtr<FText> HighlightText;
+			TListItemType RowData;
+		};
+
+		/** Generates the widget content for a column in a row */
+		DECLARE_DELEGATE_RetVal_OneParam(TSharedRef<SWidget>, FGenerateColumnWidget,
+			const FBuildArgs& InArgs
+			);
+
+		DECLARE_DELEGATE_TwoParams(FPopulateSearchString,
+			const TListItemType& InArgs,
+			TArray<FString>& InOutSearchStrings
+			);
+		
+		SLATE_BEGIN_ARGS(TReplicationColumn)
+		{}
+			/** Callback to use for generating the widget in the column */
+			SLATE_EVENT(FGenerateColumnWidget, GenerateWidgetColumn)
+			/** Callback used to generate search items for this column */
+			SLATE_EVENT(FPopulateSearchString, PopulateSearchItems)
+			/** Where in the row this column is found with respect to the other columns */
+			SLATE_ARGUMENT(int32, ColumnSortOrder)
+		SLATE_END_ARGS()
+
+		TReplicationColumn(const TReplicationColumn::FArguments& InArgs, const FColumn::FArguments& InColumnArgs)
+			: FColumn(InColumnArgs)
+			, GenerateColumnWidgetCallback(InArgs._GenerateWidgetColumn)
+			, PopulateSearchStringCallback(InArgs._PopulateSearchItems)
+			, ColumnSortOrderValue(InArgs._ColumnSortOrder)
+		{}
+
+		TSharedRef<SWidget> BuildColumnWidget(const FBuildArgs& InArgs) const
+		{
+			return ensure(GenerateColumnWidgetCallback.IsBound())
+				? GenerateColumnWidgetCallback.Execute(InArgs)
+				: SNullWidget::NullWidget;
+		}
+
+		void ExecutePopulateSearchString(const TListItemType& InRowData, TArray<FString>& InOutSearchStrings) const
+		{
+			if (PopulateSearchStringCallback.IsBound())
+			{
+				PopulateSearchStringCallback.Execute(InRowData, InOutSearchStrings);
+			}
+		}
+		
+		int32 GetColumnSortOrderValue() const { return ColumnSortOrderValue; }
+		
+	private:
+		
+		FGenerateColumnWidget GenerateColumnWidgetCallback;
+		FPopulateSearchString PopulateSearchStringCallback;
+		/** Determines whether this column is the first, etc. */
+		int32 ColumnSortOrderValue = 0;
+	};
+}
\ No newline at end of file
diff --git a/Engine/Plugins/Developer/Concert/ConcertUI/ConcertClientSharedSlate/Source/ConcertClientSharedSlate/Private/Replication/Editor/View/ObjectViewer/Tree/SReplicationColumnRow.h b/Engine/Plugins/Developer/Concert/ConcertUI/ConcertClientSharedSlate/Source/ConcertClientSharedSlate/Private/Replication/Editor/View/ObjectViewer/Tree/SReplicationColumnRow.h
new file mode 100644
--- /dev/null
+++ b/Engine/Plugins/Developer/Concert/ConcertUI/ConcertClientSharedSlate/Source/ConcertClientSharedSlate/Private/Replication/Editor/View/ObjectViewer/Tree/SReplicationColumnRow.h
@@ -0,0 +1,106 @@
+// Copyright Epic Games, Inc. All Rights Reserved.
+
+#pragma once
+
+#include "ReplicationColumn.h"
+#include "Widgets/Layout/SBox.h"
+#include "Widgets/SNullWidget.h"
+#include "Widgets/Views/SExpanderArrow.h"
+#include "Widgets/Views/STableRow.h"
+#include "Widgets/Views/STableViewBase.h"
+
+namespace UE::ConcertClientSharedSlate
+{
+	/**
+	 * Displays the summary of an activity recorded and recoverable in the SConcertSessionRecovery list view.
+	 */
+	template<typename TListItemType>
+	class SReplicationColumnRow : public SMultiColumnTableRow<TSharedPtr<TListItemType>>
+	{
+	public:
+
+		using TColumType = TReplicationColumn<TListItemType>;
+		
+		DECLARE_DELEGATE_RetVal_OneParam(const TColumType*, FGetColumn,
+			const FName& ColumnId
+			);
+		
+		SLATE_BEGIN_ARGS(SReplicationColumnRow)
+			: _RowHeight(20.f)
+		{}
+			/** Used for highlighting the text being searched. */
+			SLATE_ARGUMENT(TSharedPtr<FText>, HighlightText)
+
+			/** Gets columns info about a certain column */
+			SLATE_EVENT(FGetColumn, ColumnGetter)
+
+			/** The data to pass to TReplicationColumn::BuildColumnWidget. */
+			SLATE_ARGUMENT(TListItemType, RowData)
+		
+			/** The name of the column that will have the SExpanderArrow for the tree view. */
+			SLATE_ARGUMENT(FName, ExpandableColumnLabel)
+
+			/** The height of the row */
+			SLATE_ARGUMENT(float, RowHeight)
+		SLATE_END_ARGS()
+
+		void Construct(
+			const FArguments& InArgs,
+			TSharedRef<STableViewBase> InOwner)
+		{
+			ColumnGetterDelegate = InArgs._ColumnGetter;
+			HighlightText = InArgs._HighlightText;
+			RowData = InArgs._RowData;
+			ExpandableColumnLabel = InArgs._ExpandableColumnLabel;
+			RowHeight = InArgs._RowHeight;
+			
+			SMultiColumnTableRow<TSharedPtr<TListItemType>>::Construct({}, InOwner);
+		}
+
+		/** Generates the widget representing this row. */
+		virtual TSharedRef<SWidget> GenerateWidgetForColumn(const FName& ColumnName) override
+		{
+			const TColumType* Column = ColumnGetterDelegate.Execute(ColumnName); ensure(Column);
+			if (!Column)
+			{
+				return SNullWidget::NullWidget;
+			}
+			
+			const TSharedRef<SWidget> ColumnWidget = Column->BuildColumnWidget({ HighlightText, RowData });
+			const bool bNeedsExpanderArrow = ColumnName == ExpandableColumnLabel;
+			if (!bNeedsExpanderArrow)
+			{
+				return ColumnWidget;
+			}
+			
+			return SNew(SBox)
+				.MinDesiredHeight(RowHeight)
+				[
+					SNew(SHorizontalBox)
+
+					+SHorizontalBox::Slot()
+					.AutoWidth()
+					.Padding(6, 0, 0, 0)
+					[
+						SNew(SExpanderArrow, SReplicationColumnRow::SharedThis(this))
+						.IndentAmount(12)
+					]
+
+					+SHorizontalBox::Slot()
+					.FillWidth(1.0f)
+					[
+						ColumnWidget
+					]
+				];
+	
+		}
+
+	private:
+		
+		FGetColumn ColumnGetterDelegate;
+		TSharedPtr<FText> HighlightText;
+		TListItemType RowData;
+		FName ExpandableColumnLabel;
+		float RowHeight;
+	};
+}
\ No newline at end of file
diff --git a/Engine/Plugins/Developer/Concert/ConcertUI/ConcertClientSharedSlate/Source/ConcertClientSharedSlate/Private/Replication/Editor/View/ObjectViewer/Tree/SReplicationTreeView.h b/Engine/Plugins/Developer/Concert/ConcertUI/ConcertClientSharedSlate/Source/ConcertClientSharedSlate/Private/Replication/Editor/View/ObjectViewer/Tree/SReplicationTreeView.h
new file mode 100644
--- /dev/null
+++ b/Engine/Plugins/Developer/Concert/ConcertUI/ConcertClientSharedSlate/Source/ConcertClientSharedSlate/Private/Replication/Editor/View/ObjectViewer/Tree/SReplicationTreeView.h
@@ -0,0 +1,354 @@
+// Copyright Epic Games, Inc. All Rights Reserved.
+
+#pragma once
+
+#include "ReplicationColumn.h"
+#include "SReplicationColumnRow.h"
+
+#include "Algo/RemoveIf.h"
+#include "Misc/TextFilter.h"
+#include "Widgets/DeclarativeSyntaxSupport.h"
+#include "Widgets/Input/SSearchBox.h"
+#include "Widgets/Layout/SScrollBox.h"
+#include "Widgets/SBoxPanel.h"
+#include "Widgets/Views/STreeView.h"
+
+#define LOCTEXT_NAMESPACE "SReplicationListView"
+
+namespace UE::ConcertClientSharedSlate
+{
+	/**
+	 * Shared code for the list view for replicated actors and properties.
+	 * It is a table view that is searchable with a search box and exposes slots to add more filter widgets, such as SBasicFilterBar.
+	 */
+	template<typename TItemType>
+	class SReplicationTreeView : public SCompoundWidget
+	{
+	public:
+
+		DECLARE_DELEGATE_OneParam(FDeleteItems, const TArray<TItemType>& SelectedItems);
+		DECLARE_DELEGATE_TwoParams(FGetItemChildren, TItemType Item, TFunctionRef<void(TItemType)> ProcessChild);
+		DECLARE_DELEGATE(FOnSelectionChanged);
+		DECLARE_DELEGATE_RetVal_OneParam(bool, FCustomFilter, const TItemType& Item);
+
+		SLATE_BEGIN_ARGS(SReplicationTreeView<TItemType>)
+			: _SelectionMode(ESelectionMode::Single)
+		{}
+			/** The items to display */
+			SLATE_ARGUMENT(TArray<TItemType>*, RootItemsSource)
+
+			/** Gets an items children for the tree view */
+			SLATE_EVENT(FGetItemChildren, OnGetChildren)
+
+			/** Optional. Called when the user presses the delete key */
+			SLATE_EVENT(FDeleteItems, OnDeleteItems)
+
+			/** Called to generate the context menu for an item */
+			SLATE_EVENT(FOnContextMenuOpening, OnContextMenuOpening)
+
+			/** Called when the selection changes. Call GetSelectedItems to get the selected items. */
+			SLATE_EVENT(FOnSelectionChanged, OnSelectionChanged)
+
+			/** Optional callback to do even more filtering of items. */
+			SLATE_EVENT(FCustomFilter, FilterItem)
+			
+			/** The columns this list should have */
+			SLATE_ARGUMENT(TArray<TReplicationColumn<TItemType>>, Columns)
+			/** The name of the column that will have the SExpanderArrow for the tree view. */
+			SLATE_ARGUMENT(FName, ExpandableColumnLabel)
+		
+			/** How many items are to allowed to be selected */
+			SLATE_ARGUMENT(ESelectionMode::Type, SelectionMode)
+
+			/** Optional widget to add to the left of the search bar. */
+			SLATE_NAMED_SLOT(FArguments, LeftOfSearchBar)
+			/** Optional widget to add between the search bar and the table view (e.g. a SBasicFilterBar). */
+			SLATE_NAMED_SLOT(FArguments, RowBelowSearchBar)
+		SLATE_END_ARGS()
+
+		void Construct(const FArguments& InArgs)
+		{
+			AllRootItems = InArgs._RootItemsSource;
+			check(AllRootItems);
+
+			OnGetChildrenDelegate = InArgs._OnGetChildren;
+			OnDeleteItemsDelegate = InArgs._OnDeleteItems;
+			CustomFilterDelegate = InArgs._FilterItem;
+			ExpandableColumnId = InArgs._ExpandableColumnLabel;
+			
+			SearchText = MakeShared<FText>();
+			SearchTextFilter = MakeShared<TTextFilter<const TItemType&>>(TTextFilter<const TItemType&>::FItemToStringArray::CreateSP(this, &SReplicationTreeView::PopulateSearchStrings));
+			SearchTextFilter->OnChanged().AddSP(this, &SReplicationTreeView::OnFilterChanged);
+			
+			ChildSlot
+			[
+				SNew(SVerticalBox)
+
+				// Search
+				+SVerticalBox::Slot()
+				.Padding(1.f)
+				.AutoHeight()
+				[
+					SNew(SHorizontalBox)
+
+					+SHorizontalBox::Slot()
+					.Padding(1.f)
+					.AutoWidth()
+					[
+						InArgs._LeftOfSearchBar.Widget
+					]
+
+					+SHorizontalBox::Slot()
+					.Padding(1.f)
+					.FillWidth(1.f)
+					[
+						SAssignNew(SearchBox, SSearchBox)
+						.HintText(LOCTEXT("SearchHint", "Search..."))
+						.OnTextChanged(this, &SReplicationTreeView::OnSearchTextChanged)
+						.OnTextCommitted(this, &SReplicationTreeView::OnSearchTextCommitted)
+						.DelayChangeNotificationsWhileTyping(true)
+					]
+				]
+
+				// Optional slot between search bar and table, e.g. for an external SBasicFilterBar
+				+SVerticalBox::Slot()
+				.Padding(1.f)
+				.AutoHeight()
+				[
+					InArgs._RowBelowSearchBar.Widget
+				]
+
+				// Table row
+				+SVerticalBox::Slot()
+				.Padding(1.f)
+				.FillHeight(1.f)
+				[
+					SNew(SScrollBox)
+					+SScrollBox::Slot()
+					[
+						CreateTreeView(InArgs)
+					]
+				]
+			];
+		}
+		
+		void OnItemsChanged()
+		{
+			// Re-filter everything. There should not be many items so filtering everything again should be fine. Calls RequestListRefresh as well.
+			OnFilterChanged();
+		}
+		
+		void SetSelectedItems(const TArray<TItemType>& ObjectsToSelect, bool bIsSelected)
+		{
+			TreeView->ClearSelection();
+			TreeView->SetItemSelection(ObjectsToSelect, bIsSelected);
+		}
+		void SetExpandedItems(const TArray<TItemType>& ObjectsToSelect, bool bIsExpanded)
+		{
+			for (TItemType Item : ObjectsToSelect)
+			{
+				TreeView->SetItemExpansion(Item, bIsExpanded);
+			}
+		}
+		
+		TArray<TItemType> GetSelectedItems() const { return TreeView->GetSelectedItems(); }
+		
+		virtual FReply OnKeyDown(const FGeometry& MyGeometry, const FKeyEvent& InKeyEvent) override;
+
+	private:
+
+		/** The widget being used for search. */
+		TSharedPtr<SSearchBox> SearchBox;
+		/** Used to highlight text in text widgets */
+		TSharedPtr<FText> SearchText;
+		/** Performs text search */
+		TSharedPtr<TTextFilter<const TItemType&>> SearchTextFilter;
+
+		/** ListView's header row */
+		TSharedPtr<SHeaderRow> HeaderRow;
+		/** Displays the contents */
+		TSharedPtr<STreeView<TItemType>> TreeView;
+		/** The name of the column which will have the SExpandableArrow widget for the tree view. */
+		FName ExpandableColumnId;
+
+		TArray<TItemType>* AllRootItems = nullptr;
+		TArray<TItemType> FilteredRootItems;
+
+		/** Callback for getting an item's children. */
+		FGetItemChildren OnGetChildrenDelegate;
+		/** Optional delegate for responding to pressing the delete button */
+		FDeleteItems OnDeleteItemsDelegate;
+		/** Optional delegate for filtering the items even more. */
+		FCustomFilter CustomFilterDelegate;
+		
+
+		TSharedRef<SWidget> CreateTreeView(const FArguments& InArgs);
+		TSharedRef<SHeaderRow> CreateHeaderRow(const FArguments& InArgs);
+		TSharedRef<ITableRow> OnGenerateRowWidget(TItemType Item, const TSharedRef<STableViewBase>& OwnerTable);
+		void GetRowChildren(TItemType Item, TArray<TItemType>& OutChildren);
+		
+		void OnSearchTextCommitted(const FText& InFilterText, ETextCommit::Type CommitType);
+		void OnSearchTextChanged(const FText& InSearchText);
+
+		void PopulateSearchStrings(const TItemType& Item, TArray<FString>& OutSearchStrings);
+		void OnFilterChanged();
+		bool PassesFilters(const TItemType& Item);
+	};
+
+	template <typename TItemType>
+	FReply SReplicationTreeView<TItemType>::OnKeyDown(const FGeometry& MyGeometry, const FKeyEvent& InKeyEvent)
+	{
+		if (InKeyEvent.GetKey() == EKeys::Delete && OnDeleteItemsDelegate.IsBound())
+		{
+			OnDeleteItemsDelegate.Execute(TreeView->GetSelectedItems());
+			return FReply::Handled();
+		}
+	
+		return SCompoundWidget::OnKeyDown(MyGeometry, InKeyEvent);
+	}
+
+	template <typename TItemType>
+	TSharedRef<SWidget> SReplicationTreeView<TItemType>::CreateTreeView(const FArguments& InArgs)
+	{
+		return SNew(SBorder)
+			.BorderImage(FAppStyle::Get().GetBrush("ToolPanel.GroupBorder"))
+			.BorderBackgroundColor(FSlateColor(FLinearColor(0.6, 0.6, 0.6)))
+			.Padding(0)
+			[
+				SAssignNew(TreeView, STreeView<TItemType>)
+				.OnGetChildren(this, &SReplicationTreeView::GetRowChildren)
+				.TreeItemsSource(&FilteredRootItems)
+				.OnGenerateRow(this, &SReplicationTreeView::OnGenerateRowWidget)
+				.OnContextMenuOpening(InArgs._OnContextMenuOpening)
+				.OnSelectionChanged_Lambda([OnSelectionChanged = InArgs._OnSelectionChanged](auto, auto){ OnSelectionChanged.ExecuteIfBound(); })
+				.SelectionMode(InArgs._SelectionMode)
+				.AllowOverscroll(EAllowOverscroll::No)
+				.HeaderRow(CreateHeaderRow(InArgs))
+			];
+	}
+
+	template <typename TItemType>
+	TSharedRef<SHeaderRow> SReplicationTreeView<TItemType>::CreateHeaderRow(const FArguments& InArgs)
+	{
+		TArray<TReplicationColumn<TItemType>> Columns = InArgs._Columns;
+		Columns.Sort([](const TReplicationColumn<TItemType>& Left, const TReplicationColumn<TItemType>& Right) { return Left.GetColumnSortOrderValue() < Right.GetColumnSortOrderValue(); });
+		
+		HeaderRow = SNew(SHeaderRow);
+		TSet<FName> DuplicateColumnDetection;
+		for (TReplicationColumn<TItemType>& Column : Columns)
+		{
+			check(!DuplicateColumnDetection.Contains(Column.ColumnId));
+			DuplicateColumnDetection.Add(Column.ColumnId);
+			
+			// SHeaderRow owns the columns and deletes them when destroyed
+			TReplicationColumn<TItemType>* ManagedByHeaderRow = new TReplicationColumn<TItemType>(MoveTemp(Column));
+			HeaderRow->AddColumn(*ManagedByHeaderRow);
+		}
+
+		return HeaderRow.ToSharedRef();
+	}
+
+	template <typename TItemType>
+	TSharedRef<ITableRow> SReplicationTreeView<TItemType>::OnGenerateRowWidget(TItemType Item, const TSharedRef<STableViewBase>& OwnerTable)
+	{
+		const typename SReplicationColumnRow<TItemType>::FGetColumn ColumnGetter = SReplicationColumnRow<TItemType>::FGetColumn::CreateLambda([this](const FName& ColumnId) -> const TReplicationColumn<TItemType>*
+		{
+			for (const SHeaderRow::FColumn& Column : HeaderRow->GetColumns())
+			{
+				if (Column.ColumnId == ColumnId)
+				{
+					return static_cast<const TReplicationColumn<TItemType>*>(&Column);
+				}
+			}
+
+			checkNoEntry();
+			return nullptr;
+		});
+	
+		return SNew(SReplicationColumnRow<TItemType>, OwnerTable)
+			.HighlightText(SearchText)
+			.ColumnGetter(ColumnGetter)
+			.RowData(Item)
+			.ExpandableColumnLabel(ExpandableColumnId);
+	}
+
+	template <typename TItemType>
+	void SReplicationTreeView<TItemType>::GetRowChildren(TItemType Item, TArray<TItemType>& OutChildren)
+	{
+		if (OnGetChildrenDelegate.IsBound())
+		{
+			OnGetChildrenDelegate.Execute(Item, [this, &OutChildren](TItemType ItemToAdd)
+			{
+				if (PassesFilters(ItemToAdd))
+				{
+					OutChildren.Add(ItemToAdd);
+				}
+			});
+		}
+	}
+
+	template <typename TItemType>
+	void SReplicationTreeView<TItemType>::OnSearchTextCommitted(const FText& InFilterText, ETextCommit::Type CommitType)
+	{
+		if (!InFilterText.EqualTo(*SearchText))
+		{
+			OnSearchTextChanged(InFilterText);
+		}
+	}
+
+	template <typename TItemType>
+	void SReplicationTreeView<TItemType>::OnSearchTextChanged(const FText& InSearchText)
+	{
+		*SearchText = InSearchText;
+		SearchTextFilter->SetRawFilterText(InSearchText);
+		SearchBox->SetError(SearchTextFilter->GetFilterErrorText());
+	}
+
+	template <typename TItemType>
+	void SReplicationTreeView<TItemType>::PopulateSearchStrings(const TItemType& Item, TArray<FString>& OutSearchStrings)
+	{
+		for (const SHeaderRow::FColumn& Column : HeaderRow->GetColumns())
+		{
+			const TReplicationColumn<TItemType>& CastColumn = static_cast<const TReplicationColumn<TItemType>&>(Column);
+			CastColumn.ExecutePopulateSearchString(Item, OutSearchStrings);
+		}
+	}
+
+	template <typename TItemType>
+	void SReplicationTreeView<TItemType>::OnFilterChanged()
+	{
+		// Try preserving the selected activity.
+		TArray<TItemType> SelectedItems = TreeView->GetSelectedItems();
+
+		// Reset the list of displayed activities.
+		FilteredRootItems.Reset(AllRootItems->Num());
+
+		// Apply the filter.
+		for (const TItemType& Activity : *AllRootItems)
+		{
+			if (PassesFilters(Activity))
+			{
+				FilteredRootItems.Add(Activity);
+			}
+		}
+
+		// Restore/reset the selected activity.
+		SelectedItems.SetNum(Algo::RemoveIf(SelectedItems, [this](const TItemType& Item){ return !FilteredRootItems.Contains(Item); }));
+		if (!SelectedItems.IsEmpty())
+		{
+			TreeView->SetItemSelection(SelectedItems, true); // Restore previous selection.
+			TreeView->RequestScrollIntoView(SelectedItems[0]);
+		}
+
+		TreeView->RequestListRefresh();
+	}
+
+	template <typename TItemType>
+	bool SReplicationTreeView<TItemType>::PassesFilters(const TItemType& Item)
+	{
+		return SearchTextFilter->PassesFilter(Item)
+			&& (!CustomFilterDelegate.IsBound() || CustomFilterDelegate.Execute(Item));
+	}
+}
+
+#undef LOCTEXT_NAMESPACE
\ No newline at end of file
diff --git a/Engine/Plugins/Developer/Concert/ConcertUI/ConcertClientSharedSlate/Source/ConcertClientSharedSlate/Private/Replication/Editor/View/ObjectViewer/Tree/SelectionViewerColumns.cpp b/Engine/Plugins/Developer/Concert/ConcertUI/ConcertClientSharedSlate/Source/ConcertClientSharedSlate/Private/Replication/Editor/View/ObjectViewer/Tree/SelectionViewerColumns.cpp
new file mode 100644
--- /dev/null
+++ b/Engine/Plugins/Developer/Concert/ConcertUI/ConcertClientSharedSlate/Source/ConcertClientSharedSlate/Private/Replication/Editor/View/ObjectViewer/Tree/SelectionViewerColumns.cpp
@@ -0,0 +1,204 @@
+// Copyright Epic Games, Inc. All Rights Reserved.
+
+#include "SelectionViewerColumns.h"
+
+#include "Replication/Editor/Model/DisplayUtils.h"
+#include "Replication/Editor/Model/IEditableObjectToPropertiesModel.h"
+#include "Replication/Editor/Model/ReplicatedPropertyData.h"
+#include "Replication/Editor/Model/ReplicatedObjectData.h"
+#include "Replication/Editor/View/ObjectEditor/SObjectToPropertyEditor.h"
+#include "Replication/PropertyChainUtils.h"
+
+#include "Internationalization/Internationalization.h"
+#include "Textures/SlateIcon.h"
+#include "Widgets/Images/SImage.h"
+#include "Widgets/SBoxPanel.h"
+#include "Widgets/Text/STextBlock.h"
+#include "Widgets/Input/SCheckBox.h"
+
+#define LOCTEXT_NAMESPACE "ReplicationObjectColumns"
+
+namespace UE::ConcertClientSharedSlate::ReplicationObjectColumns
+{
+	const FName IconColumnId = TEXT("IconColumn");
+	const FName LabelColumnId = TEXT("LabelColumn");
+	const FName TypeColumnId = TEXT("TypeColumn");
+
+	FReplicationObjectColumn IconColumn(TSharedRef<IObjectToPropertiesModel> Model, const float ColumnWidth)
+	{
+		return FReplicationObjectColumn(
+			FReplicationObjectColumn::FArguments()
+				.GenerateWidgetColumn_Lambda([Model](const FReplicationObjectColumn::FBuildArgs& Args)
+				{
+					return SNew(SHorizontalBox)
+						+SHorizontalBox::Slot()
+						.HAlign(HAlign_Center)
+						.VAlign(VAlign_Center)
+						[
+							SNew(SImage)
+							.Image(DisplayUtils::GetObjectIcon(*Model, Args.RowData->GetObjectPath()).GetOptionalIcon())
+						];
+				})
+				.ColumnSortOrder(static_cast<int32>(EReplicationColumnOrder::Icon)),
+			SHeaderRow::Column(IconColumnId)
+				.DefaultLabel(FText::GetEmpty())
+				.FixedWidth(ColumnWidth)
+			);
+	}
+	
+	FReplicationObjectColumn LabelColumn()
+	{
+		return FReplicationObjectColumn(
+			FReplicationObjectColumn::FArguments()
+				.GenerateWidgetColumn_Lambda([](const FReplicationObjectColumn::FBuildArgs& Args)
+				{
+					return SNew(STextBlock)
+						.HighlightText(TAttribute<FText>::CreateLambda([HighlightText = Args.HighlightText](){ return *HighlightText; }))
+						.Text(DisplayUtils::GetObjectDisplayText(Args.RowData->GetObjectPath()));
+				})
+				.PopulateSearchItems_Lambda([](const TSharedPtr<FReplicatedObjectData>& ObjectData, TArray<FString>& InOutSearchStrings)
+				{
+					InOutSearchStrings.Add(DisplayUtils::GetObjectDisplayText(ObjectData->GetObjectPath()).ToString());
+				})
+				.ColumnSortOrder(static_cast<int32>(EReplicationColumnOrder::Label)),
+			SHeaderRow::Column(LabelColumnId)
+				.DefaultLabel(LOCTEXT("LabelColumnLabel", "Label"))
+				.FillWidth(1.f)
+			);
+	}
+	
+	FReplicationObjectColumn TypeColumn(TSharedRef<IObjectToPropertiesModel> Model)
+	{
+		return FReplicationObjectColumn(
+			FReplicationObjectColumn::FArguments()
+				.GenerateWidgetColumn_Lambda([Model](const FReplicationObjectColumn::FBuildArgs& Args)
+				{
+					return SNew(STextBlock)
+						.HighlightText(TAttribute<FText>::CreateLambda([HighlightText = Args.HighlightText](){ return *HighlightText; }))
+						.Text(DisplayUtils::GetObjectTypeText(*Model, Args.RowData->GetObjectPath()));
+				})
+				.PopulateSearchItems_Lambda([Model](const TSharedPtr<FReplicatedObjectData>& ObjectData, TArray<FString>& InOutSearchStrings)
+				{
+					InOutSearchStrings.Add(DisplayUtils::GetObjectTypeText(Model.Get(), ObjectData->GetObjectPath()).ToString());
+				})
+				.ColumnSortOrder(static_cast<int32>(EReplicationColumnOrder::Type)),
+			SHeaderRow::Column(TypeColumnId)
+				.DefaultLabel(LOCTEXT("TypeColumnLabel", "Type"))
+				.FillWidth(1.f)
+			);
+	}
+}
+
+#undef LOCTEXT_NAMESPACE
+#define LOCTEXT_NAMESPACE "ReplicationPropertyColumns"
+
+namespace UE::ConcertClientSharedSlate::ReplicationPropertyColumns
+{
+	const FName ReplicatesColumnId = TEXT("ReplicatesColumn");
+	const FName LabelColumnId = TEXT("LabelColumn");
+	const FName TypeColumnId = TEXT("TypeColumn");
+
+	FReplicationPropertyColumn ReplicatesColumns(
+		FGetPropertyCheckboxState GetPropertyCheckboxStateDelegate,
+		FOnPropertyCheckboxChanged OnPropertyBoxToggledDelegate,
+		const float ColumnWidth
+		)
+	{
+		return FReplicationPropertyColumn(
+			FReplicationPropertyColumn::FArguments()
+				.GenerateWidgetColumn_Lambda([GetPropertyCheckboxStateDelegate, OnPropertyBoxToggledDelegate](const FReplicationPropertyColumn::FBuildArgs& Args)
+				{
+					return SNew(SCheckBox)
+						.ToolTipText(LOCTEXT("ReplicatesCheckbox.Tooltip", "Should replicate?"))
+						.IsChecked_Lambda([GetPropertyCheckboxStateDelegate, RowData = Args.RowData]()
+						{
+							return GetPropertyCheckboxStateDelegate.Execute(RowData->GetProperty());
+						})
+						.OnCheckStateChanged_Lambda([OnPropertyBoxToggledDelegate, RowData = Args.RowData](ECheckBoxState NewState)
+						{
+							const bool bIsChecked = NewState == ECheckBoxState::Checked;
+							OnPropertyBoxToggledDelegate.Execute(bIsChecked, RowData->GetProperty());
+						});
+				})
+				.ColumnSortOrder(static_cast<int32>(EReplicationPropertyColumnOrder::ReplicatesCheckbox)),
+			SHeaderRow::Column(ReplicatesColumnId)
+				.DefaultLabel(FText::GetEmpty())
+				.FixedWidth(ColumnWidth)
+			);
+	}
+	
+	FReplicationPropertyColumn LabelColumn()
+	{
+		return FReplicationPropertyColumn(
+			FReplicationPropertyColumn::FArguments()
+				.GenerateWidgetColumn_Lambda([](const FReplicationPropertyColumn::FBuildArgs& Args)
+				{
+					return SNew(STextBlock)
+						.HighlightText(TAttribute<FText>::CreateLambda([HighlightText = Args.HighlightText](){ return *HighlightText; }))
+						.Text(DisplayUtils::GetPropertyDisplayText(Args.RowData->GetProperty()));
+				})
+				.PopulateSearchItems_Lambda([](const TSharedPtr<FReplicatedPropertyData>& ObjectData, TArray<FString>& InOutSearchStrings)
+				{
+					InOutSearchStrings.Add(DisplayUtils::GetPropertyDisplayText(ObjectData->GetProperty()).ToString());
+				})
+				.ColumnSortOrder(static_cast<int32>(EReplicationPropertyColumnOrder::Label)),
+			SHeaderRow::Column(LabelColumnId)
+				.DefaultLabel(LOCTEXT("LabelColumnLabel", "Label"))
+				.FillSized(500.f)
+			);
+	}
+	
+	FReplicationPropertyColumn TypeColumn()
+	{
+		static auto GetDisplayText = [](const TSharedPtr<FReplicatedPropertyData>& Args)
+		{
+			check(Args);
+			UClass* Class = Args->GetOwningClass().TryLoadClass<UObject>();
+			const FProperty* Property = Class ? ConcertSyncCore::PropertyChain::ResolveProperty(*Class, Args->GetProperty()) : nullptr;
+			return Property ? FText::FromString(Property->GetCPPType()) : LOCTEXT("Unknown", "Unknown");	
+		};
+		
+		return FReplicationPropertyColumn(
+			FReplicationPropertyColumn::FArguments()
+				.GenerateWidgetColumn_Lambda([](const FReplicationPropertyColumn::FBuildArgs& Args)
+				{
+					return SNew(STextBlock)
+						.HighlightText(TAttribute<FText>::CreateLambda([HighlightText = Args.HighlightText](){ return *HighlightText; }))
+						.Text(GetDisplayText(Args.RowData));
+				})
+				.PopulateSearchItems_Lambda([](const TSharedPtr<FReplicatedPropertyData>& ObjectData, TArray<FString>& InOutSearchStrings)
+				{
+					InOutSearchStrings.Add(GetDisplayText(ObjectData).ToString());
+				})
+				.ColumnSortOrder(static_cast<int32>(EReplicationPropertyColumnOrder::Type)),
+			SHeaderRow::Column(TypeColumnId)
+				.DefaultLabel(LOCTEXT("TypeColumnLabel", "Type"))
+				.FillWidth(1.f)
+			);
+	}
+
+	ECheckBoxState GetPropertyCheckboxStateBasedOnSelection(const FConcertPropertyChain& Property, TConstArrayView<FSoftObjectPath> Selection, const IObjectToPropertiesModel& Model)
+	{
+		ECheckBoxState CheckBoxState = ECheckBoxState::Undetermined;
+		for (const FSoftObjectPath& SelectedObject : Selection)
+		{
+			const bool bContainsProperty = Model.ContainsProperties(SelectedObject, { Property });
+			const ECheckBoxState StateForThisObject = bContainsProperty
+				? ECheckBoxState::Checked
+				: ECheckBoxState::Unchecked;
+			// The first object?
+			if (CheckBoxState == ECheckBoxState::Undetermined)
+			{
+				CheckBoxState = StateForThisObject;
+			}
+			else if (CheckBoxState != StateForThisObject)
+			{
+				// Two boxes do not have the same value
+				return ECheckBoxState::Undetermined;
+			}
+		}
+		return CheckBoxState;
+	}
+}
+
+#undef LOCTEXT_NAMESPACE
\ No newline at end of file
diff --git a/Engine/Plugins/Developer/Concert/ConcertUI/ConcertClientSharedSlate/Source/ConcertClientSharedSlate/Private/Replication/Editor/View/ObjectViewer/Tree/SelectionViewerColumns.h b/Engine/Plugins/Developer/Concert/ConcertUI/ConcertClientSharedSlate/Source/ConcertClientSharedSlate/Private/Replication/Editor/View/ObjectViewer/Tree/SelectionViewerColumns.h
new file mode 100644
--- /dev/null
+++ b/Engine/Plugins/Developer/Concert/ConcertUI/ConcertClientSharedSlate/Source/ConcertClientSharedSlate/Private/Replication/Editor/View/ObjectViewer/Tree/SelectionViewerColumns.h
@@ -0,0 +1,80 @@
+// Copyright Epic Games, Inc. All Rights Reserved.
+
+#pragma once
+
+#include "Containers/ArrayView.h"
+#include "ReplicationColumn.h"
+#include "Templates/SharedPointer.h"
+
+struct FConcertPropertyChain;
+
+namespace UE::ConcertClientSharedSlate
+{
+	class FReplicatedObjectData;
+	class FReplicatedPropertyData;
+	class IEditableObjectToPropertiesModel;
+	class IObjectToPropertiesModel;
+	class SObjectToPropertyEditor;
+}
+
+namespace UE::ConcertClientSharedSlate::ReplicationObjectColumns
+{
+	using FReplicationObjectColumn = TReplicationColumn<TSharedPtr<FReplicatedObjectData>>;
+	
+	enum class EReplicationColumnOrder : int32
+	{
+		/** Displays the class icon */
+		Icon = 10,
+		/** Label of the object */
+		Label = 20,
+		/** Class of the object */
+		Type = 30,
+	};
+	
+	extern const FName IconColumnId;
+	extern const FName LabelColumnId;
+	extern const FName TypeColumnId;
+
+	FReplicationObjectColumn IconColumn(TSharedRef<IObjectToPropertiesModel> Model, const float ColumnWidth = 20.f);
+	FReplicationObjectColumn LabelColumn();
+	FReplicationObjectColumn TypeColumn(TSharedRef<IObjectToPropertiesModel> Model);
+}
+
+namespace UE::ConcertClientSharedSlate::ReplicationPropertyColumns
+{
+	using FReplicationPropertyColumn = TReplicationColumn<TSharedPtr<FReplicatedPropertyData>>;
+	
+	enum class EReplicationPropertyColumnOrder : int32
+	{
+		/** The checkbox in SObjectToPropertyEditor determining whether the property is in the selection*/
+		ReplicatesCheckbox = 0,
+		/** Label of the property */
+		Label = 10,
+		/** Type of the property */
+		Type = 20
+	};
+	
+	/** The checkbox in SObjectToPropertyEditor determining whether the property is in the selection*/
+	extern const FName ReplicatesColumnId;
+	extern const FName LabelColumnId;
+	extern const FName TypeColumnId;
+
+	DECLARE_DELEGATE_RetVal_OneParam(ECheckBoxState, FGetPropertyCheckboxState, const FConcertPropertyChain& /*Property*/);
+	DECLARE_DELEGATE_TwoParams(FOnPropertyCheckboxChanged, bool /*bIsChecked*/, const FConcertPropertyChain& /*Property*/);
+
+	FReplicationPropertyColumn ReplicatesColumns(
+		FGetPropertyCheckboxState GetPropertyCheckboxStateDelegate,
+		FOnPropertyCheckboxChanged OnPropertyBoxToggledDelegate,
+		const float ColumnWidth = 20.f
+		);
+	FReplicationPropertyColumn LabelColumn();
+	FReplicationPropertyColumn TypeColumn();
+
+	/**
+	 * Goes through all selected objects, checks whether the property is checked on it or not, and returns a checkbox state.
+	 * The state is checked or unchecked if all selected properties have the property enabled or disabled, respectively, and undetermined otherwise, i.e. if it is mixed.
+	 *
+	 * This is used by ReplicatesColumns.
+	 */
+	ECheckBoxState GetPropertyCheckboxStateBasedOnSelection(const FConcertPropertyChain& Property, TConstArrayView<FSoftObjectPath> Selection, const IObjectToPropertiesModel& Model);
+}
diff --git a/Engine/Plugins/Developer/Concert/ConcertUI/ConcertClientSharedSlate/Source/ConcertClientSharedSlate/Private/Replication/Editor/View/SReplicationStreamEditor.cpp b/Engine/Plugins/Developer/Concert/ConcertUI/ConcertClientSharedSlate/Source/ConcertClientSharedSlate/Private/Replication/Editor/View/SReplicationStreamEditor.cpp
--- a/Engine/Plugins/Developer/Concert/ConcertUI/ConcertClientSharedSlate/Source/ConcertClientSharedSlate/Private/Replication/Editor/View/SReplicationStreamEditor.cpp
+++ b/Engine/Plugins/Developer/Concert/ConcertUI/ConcertClientSharedSlate/Source/ConcertClientSharedSlate/Private/Replication/Editor/View/SReplicationStreamEditor.cpp
@@ -3,7 +3,7 @@
 #include "SReplicationStreamEditor.h"
 
 #include "Replication/ReplicationWidgetFactories.h"
-#include "Replication/Editor/View/ObjectEditor/SPropertyReplicationSelectionEditor.h"
+#include "Replication/Editor/View/ObjectEditor/SObjectToPropertyEditor.h"
 
 #include "Widgets/SBoxPanel.h"
 
@@ -17,7 +17,8 @@ namespace UE::ConcertClientSharedSlate
 
 			+SVerticalBox::Slot()
 			[
-				SAssignNew(Editor, SPropertyReplicationSelectionEditor, Params.DataModel, Params.ObjectSource, Params.PropertySource)
+				SAssignNew(Editor, SObjectToPropertyEditor, Params.DataModel, Params.ObjectSource, Params.PropertySource)
+				.SubobjectView(Params.SubobjectView)
 			]
 		];
 	}
@@ -25,6 +26,7 @@ namespace UE::ConcertClientSharedSlate
 	void SReplicationStreamEditor::Refresh()
 	{
 		Editor->RefreshObjectData();
+		Editor->RefreshSubobjectData();
 		Editor->RefreshPropertyData();
 	}
 }
diff --git a/Engine/Plugins/Developer/Concert/ConcertUI/ConcertClientSharedSlate/Source/ConcertClientSharedSlate/Private/Replication/Editor/View/SReplicationStreamEditor.h b/Engine/Plugins/Developer/Concert/ConcertUI/ConcertClientSharedSlate/Source/ConcertClientSharedSlate/Private/Replication/Editor/View/SReplicationStreamEditor.h
--- a/Engine/Plugins/Developer/Concert/ConcertUI/ConcertClientSharedSlate/Source/ConcertClientSharedSlate/Private/Replication/Editor/View/SReplicationStreamEditor.h
+++ b/Engine/Plugins/Developer/Concert/ConcertUI/ConcertClientSharedSlate/Source/ConcertClientSharedSlate/Private/Replication/Editor/View/SReplicationStreamEditor.h
@@ -10,7 +10,7 @@ class UMultiUserReplicationClientPreset;
 
 namespace UE::ConcertClientSharedSlate
 {
-	class SPropertyReplicationSelectionEditor;
+	class SObjectToPropertyEditor;
 	
 	struct FCreateEditorParams;
 	
@@ -31,6 +31,6 @@ namespace UE::ConcertClientSharedSlate
 
 	private:
 		
-		TSharedPtr<SPropertyReplicationSelectionEditor> Editor;
+		TSharedPtr<SObjectToPropertyEditor> Editor;
 	};
 }
diff --git a/Engine/Plugins/Developer/Concert/ConcertUI/ConcertClientSharedSlate/Source/ConcertClientSharedSlate/Private/Replication/Editor/View/SubobjectView/SConcertReplicationSubobjectEditor.cpp b/Engine/Plugins/Developer/Concert/ConcertUI/ConcertClientSharedSlate/Source/ConcertClientSharedSlate/Private/Replication/Editor/View/SubobjectView/SConcertReplicationSubobjectEditor.cpp
new file mode 100644
--- /dev/null
+++ b/Engine/Plugins/Developer/Concert/ConcertUI/ConcertClientSharedSlate/Source/ConcertClientSharedSlate/Private/Replication/Editor/View/SubobjectView/SConcertReplicationSubobjectEditor.cpp
@@ -0,0 +1,108 @@
+// Copyright Epic Games, Inc. All Rights Reserved.
+
+#include "SConcertReplicationSubobjectEditor.h"
+
+#include "SubobjectDataHandle.h"
+#include "Widgets/Input/SSearchBox.h"
+#include "Widgets/SBoxPanel.h"
+#include "Widgets/SNullWidget.h"
+
+#define LOCTEXT_NAMESPACE "SConcertReplicationSubobjectEditor"
+
+namespace UE::ConcertClientSharedSlate
+{
+	void SConcertReplicationSubobjectEditor::Construct(const FArguments& InArgs)
+	{
+		OnSubobjectsSelectedDelegate = InArgs._OnSubobjectsSelected;
+		
+		OnSelectionUpdated.BindSP(this, &SConcertReplicationSubobjectEditor::HandleSelectionUpdated);
+		ObjectContext = TAttribute<UObject*>::CreateSP(this, &SConcertReplicationSubobjectEditor::GetContextObject);
+		AllowEditing = false;
+		bAllowTreeUpdates = true;
+		// No commands for now... but parent class requires this or it will crash
+		CommandList = MakeShareable(new FUICommandList);
+		ConstructTreeWidget();
+		
+		ButtonBox = SNew(SHorizontalBox)
+		+ SHorizontalBox::Slot()
+		.VAlign(VAlign_Center)
+		.Padding(0.0f)
+		.FillWidth(1.f)
+		[
+			SAssignNew(FilterBox, SSearchBox)
+			.HintText(LOCTEXT("SearchHint", "Search"))
+			.OnTextChanged(this, &SConcertReplicationSubobjectEditor::OnFilterTextChanged)
+		];
+
+		ChildSlot
+		[
+			SNew(SVerticalBox)
+
+			+SVerticalBox::Slot()
+			.AutoHeight()
+			.VAlign(VAlign_Top)
+			.Padding(4.f, 0, 4.f, 4.f)
+			[
+				ButtonBox.ToSharedRef()
+			]
+			
+			+SVerticalBox::Slot()
+			[
+				SNew(SBorder)
+				.BorderImage(FAppStyle::Get().GetBrush("SCSEditor.Background"))
+				.Padding(4.f)
+				[
+					TreeWidget.ToSharedRef()
+				]
+			]
+		];
+	}
+
+	void SConcertReplicationSubobjectEditor::SetDisplayedRootObject(UObject* Object)
+	{
+		DisplayedRootObject = Object;
+		UpdateTree();
+	}
+
+	TSharedPtr<SWidget> SConcertReplicationSubobjectEditor::BuildSceneRootDropActionMenu(FSubobjectEditorTreeNodePtrType DroppedOntoNodePtr, FSubobjectEditorTreeNodePtrType DroppedNodePtr)
+	{
+		return SNullWidget::NullWidget;
+	}
+
+	FSubobjectDataHandle SConcertReplicationSubobjectEditor::AddNewSubobject(
+		const FSubobjectDataHandle& ParentHandle,
+		UClass* NewClass,
+		UObject* AssetOverride,
+		FText& OutFailReason,
+		TUniquePtr<FScopedTransaction> InOngoingTransaction
+		)
+	{
+		// We do not edit the hierarchy. Ever. See AllowEditing being set to false in Construct.
+		checkNoEntry();
+		return FSubobjectDataHandle{};
+	}
+
+	UObject* SConcertReplicationSubobjectEditor::GetContextObject() const
+	{
+		return DisplayedRootObject.Get();
+	}
+
+	void SConcertReplicationSubobjectEditor::HandleSelectionUpdated(const TArray<TSharedPtr<FSubobjectEditorTreeNode>>& SelectedNodes) const
+	{
+		OnSubobjectsSelectedDelegate.Execute();
+	}
+
+	TSet<const UObject*> SConcertReplicationSubobjectEditor::GetSelectedObjectsFromNodes(const TArray<TSharedPtr<FSubobjectEditorTreeNode>>& SelectedNodes)
+	{
+		TSet<const UObject*> ObjectPaths;
+		Algo::Transform(SelectedNodes, ObjectPaths, [](const TSharedPtr<FSubobjectEditorTreeNode>& Node)
+		{
+			return Node->GetObject();
+		});
+
+		ObjectPaths.Remove(nullptr);
+		return ObjectPaths;
+	}
+}
+
+#undef LOCTEXT_NAMESPACE
\ No newline at end of file
diff --git a/Engine/Plugins/Developer/Concert/ConcertUI/ConcertClientSharedSlate/Source/ConcertClientSharedSlate/Private/Replication/Editor/View/SubobjectView/SConcertReplicationSubobjectEditor.h b/Engine/Plugins/Developer/Concert/ConcertUI/ConcertClientSharedSlate/Source/ConcertClientSharedSlate/Private/Replication/Editor/View/SubobjectView/SConcertReplicationSubobjectEditor.h
new file mode 100644
--- /dev/null
+++ b/Engine/Plugins/Developer/Concert/ConcertUI/ConcertClientSharedSlate/Source/ConcertClientSharedSlate/Private/Replication/Editor/View/SubobjectView/SConcertReplicationSubobjectEditor.h
@@ -0,0 +1,66 @@
+// Copyright Epic Games, Inc. All Rights Reserved.
+
+#pragma once
+
+#include "SSubobjectEditor.h"
+
+namespace UE::ConcertClientSharedSlate
+{
+	/** Extends SSubobjectEditor to display components on the actor selected in the replication outliner. */
+	class SConcertReplicationSubobjectEditor : public SSubobjectEditor
+	{
+	public:
+
+		DECLARE_DELEGATE(FOnSubobjectsSelected);
+		
+		SLATE_BEGIN_ARGS(SConcertReplicationSubobjectEditor)
+		{}
+			/** Delegate to invoke on selection update. */
+			SLATE_EVENT(FOnSubobjectsSelected, OnSubobjectsSelected)
+		SLATE_END_ARGS()
+
+		void Construct(const FArguments& InArgs);
+
+		/** Sets the object to display subjects for */
+		void SetDisplayedRootObject(UObject* Object = nullptr);
+		/** Gets the objects currently selected. */
+		TSet<const UObject*> GetSelectedObjects() const { return GetSelectedObjectsFromNodes(GetSelectedNodes()); }
+
+		//~ Begin SSubobjectEditor Interface
+		virtual void OnAttachToDropAction(FSubobjectEditorTreeNodePtrType DroppedOn,const TArray<FSubobjectEditorTreeNodePtrType>& DroppedNodePtrs) override {}
+		virtual void OnDetachFromDropAction(const TArray<FSubobjectEditorTreeNodePtrType>& DroppedNodePtrs) override {}
+		virtual void OnMakeNewRootDropAction(FSubobjectEditorTreeNodePtrType DroppedNodePtr) override {}
+		virtual void PostDragDropAction(bool bRegenerateTreeNodes) override {}
+		virtual TSharedPtr<SWidget> BuildSceneRootDropActionMenu(FSubobjectEditorTreeNodePtrType DroppedOntoNodePtr, FSubobjectEditorTreeNodePtrType DroppedNodePtr) override;
+		virtual void OnDeleteNodes() override {}
+		//~ End SSubobjectEditor Interface
+
+	protected:
+
+		//~ Begin SSubobjectEditor Interface
+		virtual void CopySelectedNodes() override {}
+		virtual void PasteNodes() override {}
+		virtual void OnDuplicateComponent() override {}
+		virtual void PopulateContextMenuImpl(UToolMenu* InMenu, TArray<FSubobjectEditorTreeNodePtrType>& InSelectedItems, bool bIsChildActorSubtreeNodeSelected) override {}
+		virtual FSubobjectDataHandle AddNewSubobject(const FSubobjectDataHandle& ParentHandle, UClass* NewClass, UObject* AssetOverride, FText& OutFailReason, TUniquePtr<FScopedTransaction> InOngoingTransaction) override;
+		//~ End SSubobjectEditor Interface
+		
+	private:
+		
+		/** The object the subobject hierarchy is displayed for */
+		TWeakObjectPtr<UObject> DisplayedRootObject;
+
+		/** Delegate to invoke on selection update. */
+		FOnSubobjectsSelected OnSubobjectsSelectedDelegate;
+		
+		/** @return The object to display the subobject hierarchy for. */
+		UObject* GetContextObject() const;
+
+		/** Utils for converting node selection to objects. */
+		static TSet<const UObject*> GetSelectedObjectsFromNodes(const TArray<TSharedPtr<FSubobjectEditorTreeNode>>& SelectedNodes);
+
+		/** Called when the selection is updated */
+		void HandleSelectionUpdated(const TArray<TSharedPtr<FSubobjectEditorTreeNode>>& SelectedNodes) const;
+	};
+}
+
diff --git a/Engine/Plugins/Developer/Concert/ConcertUI/ConcertClientSharedSlate/Source/ConcertClientSharedSlate/Private/Replication/Editor/View/SubobjectView/SUnrealEditorSubobjectView.cpp b/Engine/Plugins/Developer/Concert/ConcertUI/ConcertClientSharedSlate/Source/ConcertClientSharedSlate/Private/Replication/Editor/View/SubobjectView/SUnrealEditorSubobjectView.cpp
new file mode 100644
--- /dev/null
+++ b/Engine/Plugins/Developer/Concert/ConcertUI/ConcertClientSharedSlate/Source/ConcertClientSharedSlate/Private/Replication/Editor/View/SubobjectView/SUnrealEditorSubobjectView.cpp
@@ -0,0 +1,48 @@
+// Copyright Epic Games, Inc. All Rights Reserved.
+
+#include "SUnrealEditorSubobjectView.h"
+
+#include "SConcertReplicationSubobjectEditor.h"
+
+namespace UE::ConcertClientSharedSlate
+{
+	void SUnrealEditorSubobjectView::Construct(const FArguments& InArgs)
+	{
+		ChildSlot
+		[
+			SAssignNew(SubobjectEditor, SConcertReplicationSubobjectEditor)
+			.OnSubobjectsSelected(this, &SUnrealEditorSubobjectView::OnSubobjectsSelected)
+		];
+	}
+	
+	void SUnrealEditorSubobjectView::SetRootObjects(const TArray<FSoftObjectPath>& RootObject)
+	{
+		if (RootObject.IsEmpty())
+		{
+			SubobjectEditor->SetDisplayedRootObject(nullptr);
+		}
+		else
+		{
+			// TODO DP: Check whether all of the objects are compatible with each other
+			UObject* Object = RootObject[0].ResolveObject();
+			SubobjectEditor->SetDisplayedRootObject(Object);
+		}
+	}
+
+	void SUnrealEditorSubobjectView::SelectRootObjects()
+	{
+		SubobjectEditor->SelectRoot();
+	}
+
+	TArray<FSoftObjectPath> SUnrealEditorSubobjectView::GetSelectedObjects() const
+	{
+		TArray<FSoftObjectPath> Paths;
+		Algo::Transform(SubobjectEditor->GetSelectedObjects(), Paths, [](const UObject* Object){ return Object; });
+		return Paths;
+	}
+
+	void SUnrealEditorSubobjectView::OnSubobjectsSelected()
+	{
+		OnSelectionChangedDelegate.Broadcast();
+	}
+}
diff --git a/Engine/Plugins/Developer/Concert/ConcertUI/ConcertClientSharedSlate/Source/ConcertClientSharedSlate/Private/Replication/Editor/View/SubobjectView/SUnrealEditorSubobjectView.h b/Engine/Plugins/Developer/Concert/ConcertUI/ConcertClientSharedSlate/Source/ConcertClientSharedSlate/Private/Replication/Editor/View/SubobjectView/SUnrealEditorSubobjectView.h
new file mode 100644
--- /dev/null
+++ b/Engine/Plugins/Developer/Concert/ConcertUI/ConcertClientSharedSlate/Source/ConcertClientSharedSlate/Private/Replication/Editor/View/SubobjectView/SUnrealEditorSubobjectView.h
@@ -0,0 +1,41 @@
+// Copyright Epic Games, Inc. All Rights Reserved.
+
+#pragma once
+
+#include "Replication/Editor/View/IReplicationSubobjectView.h"
+
+namespace UE::ConcertClientSharedSlate
+{
+	class SConcertReplicationSubobjectEditor;
+
+	/** Adapts SSubobjectEditor to the IReplicationSubobjectView interface. */
+	class SUnrealEditorSubobjectView : public IReplicationSubobjectView
+	{
+	public:
+
+		SLATE_BEGIN_ARGS(SUnrealEditorSubobjectView)
+		{}
+		SLATE_END_ARGS()
+
+		void Construct(const FArguments& InArgs);
+
+		//~ Begin IReplicationSubobjectView Interface
+		virtual void SetRootObjects(const TArray<FSoftObjectPath>& RootObjects) override;
+		virtual void SelectRootObjects() override;
+		virtual TArray<FSoftObjectPath> GetSelectedObjects() const override;
+		virtual FOnSelectionChanged& OnSelectionChanged() override { return OnSelectionChangedDelegate; }
+		//~ End IReplicationSubobjectView Interface
+
+	private:
+
+		/** Reuses the SSubobjectEditor for displaying subobjects. */
+		TSharedPtr<SConcertReplicationSubobjectEditor> SubobjectEditor;
+
+		/** Executes when the selection is changed. */
+		FOnSelectionChanged OnSelectionChangedDelegate;
+
+		/** Called when the SSubobjectEditor reports new selected objects. */
+		void OnSubobjectsSelected();
+	};
+}
+
diff --git a/Engine/Plugins/Developer/Concert/ConcertUI/ConcertClientSharedSlate/Source/ConcertClientSharedSlate/Private/Replication/ReplicationWidgetFactories.cpp b/Engine/Plugins/Developer/Concert/ConcertUI/ConcertClientSharedSlate/Source/ConcertClientSharedSlate/Private/Replication/ReplicationWidgetFactories.cpp
--- a/Engine/Plugins/Developer/Concert/ConcertUI/ConcertClientSharedSlate/Source/ConcertClientSharedSlate/Private/Replication/ReplicationWidgetFactories.cpp
+++ b/Engine/Plugins/Developer/Concert/ConcertUI/ConcertClientSharedSlate/Source/ConcertClientSharedSlate/Private/Replication/ReplicationWidgetFactories.cpp
@@ -4,9 +4,15 @@
 
 #include "Editor/Model/TransactionalPropertySelectionModel.h"
 #include "Editor/View/SReplicationStreamEditor.h"
+#include "Editor/View/SubobjectView/SUnrealEditorSubobjectView.h"
 
 namespace UE::ConcertClientSharedSlate
 {
+	TSharedRef<IReplicationSubobjectView> CreateUnrealEditorSubobjectView()
+	{
+		return SNew(SUnrealEditorSubobjectView);
+	}
+	
 	TSharedRef<IReplicationEditorView> CreateEditor(FCreateEditorParams Params)
 	{
 		return SNew(SReplicationStreamEditor, MoveTemp(Params));
diff --git a/Engine/Plugins/Developer/Concert/ConcertUI/ConcertClientSharedSlate/Source/ConcertClientSharedSlate/Public/Replication/Editor/Model/IEditableObjectToPropertiesModel.h b/Engine/Plugins/Developer/Concert/ConcertUI/ConcertClientSharedSlate/Source/ConcertClientSharedSlate/Public/Replication/Editor/Model/IEditableObjectToPropertiesModel.h
--- a/Engine/Plugins/Developer/Concert/ConcertUI/ConcertClientSharedSlate/Source/ConcertClientSharedSlate/Public/Replication/Editor/Model/IEditableObjectToPropertiesModel.h
+++ b/Engine/Plugins/Developer/Concert/ConcertUI/ConcertClientSharedSlate/Source/ConcertClientSharedSlate/Public/Replication/Editor/Model/IEditableObjectToPropertiesModel.h
@@ -2,6 +2,7 @@
 
 #pragma once
 
+#include "Containers/ContainersFwd.h"
 #include "Containers/Map.h"
 #include "Containers/Set.h"
 #include "Delegates/DelegateCombinations.h"
@@ -29,17 +30,17 @@ namespace UE::ConcertClientSharedSlate
 	public:
 
 		/** Adds the objects to the mapping */
-		virtual void AddObjects(TArrayView<UObject*> Objects) = 0;
+		virtual void AddObjects(TConstArrayView<UObject*> Objects) = 0;
 		/** Removes the objects from the mapping */
-		virtual void RemoveObjects(TArrayView<FSoftObjectPath> Objects) = 0;
+		virtual void RemoveObjects(TConstArrayView<FSoftObjectPath> Objects) = 0;
 
 		/** Adds these properties to the object's list of selected properties. */
-		virtual void AddProperties(const FSoftObjectPath& Object, TArrayView<FConcertPropertyChain> Properties) = 0;
+		virtual void AddProperties(const FSoftObjectPath& Object, TConstArrayView<FConcertPropertyChain> Properties) = 0;
 		/** Removes these properties from the object's list of selected properties. */
-		virtual void RemoveProperties(const FSoftObjectPath& Object, TArrayView<FConcertPropertyChain> Properties) = 0;
+		virtual void RemoveProperties(const FSoftObjectPath& Object, TConstArrayView<FConcertPropertyChain> Properties) = 0;
 
 		/** Called when the object list changes. AddedObjects and RemovedObjects are empty if and only if ChangeReason == EReplicatedObjectChangeReason::Transacted. */
-		DECLARE_MULTICAST_DELEGATE_ThreeParams(FOnObjectsChanged, TArrayView<UObject*> AddedObjects, TArrayView<FSoftObjectPath> RemovedObjects, EReplicatedObjectChangeReason ChangeReason);
+		DECLARE_MULTICAST_DELEGATE_ThreeParams(FOnObjectsChanged, TConstArrayView<UObject*> AddedObjects, TConstArrayView<FSoftObjectPath> RemovedObjects, EReplicatedObjectChangeReason ChangeReason);
 		virtual FOnObjectsChanged& OnObjectsChanged() = 0;
 
 		/** Called when the property list of some object changes. If the properties are removed because the object is removed outright, FOnObjectsChanged is called instead.  */
diff --git a/Engine/Plugins/Developer/Concert/ConcertUI/ConcertClientSharedSlate/Source/ConcertClientSharedSlate/Public/Replication/Editor/Model/Object/AddOuterSource.h b/Engine/Plugins/Developer/Concert/ConcertUI/ConcertClientSharedSlate/Source/ConcertClientSharedSlate/Public/Replication/Editor/Model/Object/AddOuterSource.h
deleted file mode 100644
diff --git a/Engine/Plugins/Developer/Concert/ConcertUI/ConcertClientSharedSlate/Source/ConcertClientSharedSlate/Public/Replication/Editor/Model/Object/ComponentFromActorSource.h b/Engine/Plugins/Developer/Concert/ConcertUI/ConcertClientSharedSlate/Source/ConcertClientSharedSlate/Public/Replication/Editor/Model/Object/ComponentFromActorSource.h
deleted file mode 100644
diff --git a/Engine/Plugins/Developer/Concert/ConcertUI/ConcertClientSharedSlate/Source/ConcertClientSharedSlate/Public/Replication/Editor/Model/Object/EditorObjectSelectionSourceModel.h b/Engine/Plugins/Developer/Concert/ConcertUI/ConcertClientSharedSlate/Source/ConcertClientSharedSlate/Public/Replication/Editor/Model/Object/EditorObjectSelectionSourceModel.h
--- a/Engine/Plugins/Developer/Concert/ConcertUI/ConcertClientSharedSlate/Source/ConcertClientSharedSlate/Public/Replication/Editor/Model/Object/EditorObjectSelectionSourceModel.h
+++ b/Engine/Plugins/Developer/Concert/ConcertUI/ConcertClientSharedSlate/Source/ConcertClientSharedSlate/Public/Replication/Editor/Model/Object/EditorObjectSelectionSourceModel.h
@@ -25,26 +25,8 @@ namespace UE::ConcertClientSharedSlate
 
 	private:
 
-		/** Used to build Add Component source from BaseActorCategory. If you move this property, make sure it declared before BaseActorCategory, or update the constructor. */
-		TSharedRef<FWorldActorSource> WorldActorSource;
-
 		/** Contains actions around adding actor related objects. Displayed next to the search bar. Refined with dynamic options in GetRootSources. */
 		FObjectSourceCategory BaseActorCategory;
-
-		/** Sources which have to be instantiated with a context. */
-		struct FContextMenuSources
-		{
-			/** Add components from the actor for which the context menu is instantiated. */
-			TSharedRef<FComponentFromActorSource_ContextMenu> AddComponent;
-
-			/** Adds the object's outer as object. */
-			TSharedRef<FAddOuterSource> AddOuter;
-			
-			FContextMenuSources();
-
-			/** Gets all applicable context menu sources. */
-			TArray<TSharedRef<IObjectSourceModel>> BuildContextSources(const FSoftObjectPath& Item);
-		} ContextMenuSources;
 	};
 }
 
diff --git a/Engine/Plugins/Developer/Concert/ConcertUI/ConcertClientSharedSlate/Source/ConcertClientSharedSlate/Public/Replication/Editor/View/IReplicationSubobjectView.h b/Engine/Plugins/Developer/Concert/ConcertUI/ConcertClientSharedSlate/Source/ConcertClientSharedSlate/Public/Replication/Editor/View/IReplicationSubobjectView.h
new file mode 100644
--- /dev/null
+++ b/Engine/Plugins/Developer/Concert/ConcertUI/ConcertClientSharedSlate/Source/ConcertClientSharedSlate/Public/Replication/Editor/View/IReplicationSubobjectView.h
@@ -0,0 +1,51 @@
+// Copyright Epic Games, Inc. All Rights Reserved.
+
+#pragma once
+
+#include "Templates/SharedPointer.h"
+#include "Widgets/SCompoundWidget.h"
+
+#include "Containers/ArrayView.h"
+
+namespace UE::ConcertClientSharedSlate
+{
+	/**
+	 * A widget for selecting subobjects on a given root object.
+	 * When a user clicks a root entity (usually an actor) this view shows its subobjects (usually components).
+	 * 
+	 * This is an optional widget injected between the root object selection view and the property selection view.
+	 *
+	 * In the editor, this widget is backed by SSubobjectEditor.
+	 * On the server, a custom widget is necessary since certain classes like SSubobjectEditor and UActorComponent, etc.
+	 * are unavailable.
+	 * 
+	 * @see ReplicationWidgetFactories.h
+	 */
+	class CONCERTCLIENTSHAREDSLATE_API IReplicationSubobjectView : public SCompoundWidget
+	{
+	public:
+
+		/** Sets the root objects for which the subobject list should be shown. */
+		virtual void SetRootObjects(const TArray<FSoftObjectPath>& RootObjects) = 0;
+		/** Clears the hierarchy so it displays no subobjects. */
+		void ClearRootObjects() { SetRootObjects({}); }
+
+		/** Selects the root object in the UI. Executes OnSelectionChanged. */
+		virtual void SelectRootObjects() = 0;
+		
+		/**
+		 * Gets the list of selected objects.
+		 *
+		 * This may return objects that are not currently in the object replication mapping of the underlying model;
+		 * this is usually done when the surrounding UI edits the replication data. If the UI is view-only,
+		 * GetSelectedObjects should only return objects in the model.
+		 */
+		virtual TArray<FSoftObjectPath> GetSelectedObjects() const = 0;
+
+		/** Called when the subobject selection changes. */
+		DECLARE_MULTICAST_DELEGATE(FOnSelectionChanged);
+		virtual FOnSelectionChanged& OnSelectionChanged() = 0;
+		
+		virtual ~IReplicationSubobjectView() = default;
+	};
+}
\ No newline at end of file
diff --git a/Engine/Plugins/Developer/Concert/ConcertUI/ConcertClientSharedSlate/Source/ConcertClientSharedSlate/Public/Replication/ReplicationWidgetFactories.h b/Engine/Plugins/Developer/Concert/ConcertUI/ConcertClientSharedSlate/Source/ConcertClientSharedSlate/Public/Replication/ReplicationWidgetFactories.h
--- a/Engine/Plugins/Developer/Concert/ConcertUI/ConcertClientSharedSlate/Source/ConcertClientSharedSlate/Public/Replication/ReplicationWidgetFactories.h
+++ b/Engine/Plugins/Developer/Concert/ConcertUI/ConcertClientSharedSlate/Source/ConcertClientSharedSlate/Public/Replication/ReplicationWidgetFactories.h
@@ -11,6 +11,11 @@ class SWidget;
 struct FConcertReplicationEditorSettings;
 struct FObjectReplicationMap;
 
+namespace UE::ConcertClientSharedSlate
+{
+	class IReplicationSubobjectView;
+}
+
 namespace UE::ConcertClientSharedSlate
 {
 	class IEditableObjectToPropertiesModel;
@@ -39,18 +44,36 @@ namespace UE::ConcertClientSharedSlate
 		TSharedRef<IPropertySelectionSourceModel> PropertySource;
 
 		// TODO DP: Add way to add more columns
+
+		/**
+		 * Optional. This is inserted between the root object outliner and property view.
+		 * It e.g. displays the components of the actor selected in the root object view.
+		 * 
+		 * Exists so it can be customized differently depending on whether used in the editor or on the server-
+		 */
+		TSharedPtr<IReplicationSubobjectView> SubobjectView;
 	};
+	
+	/** Creates an object that looks like the SSubobjectEditor. */
+	CONCERTCLIENTSHAREDSLATE_API TSharedRef<IReplicationSubobjectView> CreateUnrealEditorSubobjectView();
 
 	/**
 	 * Creates a replication editor.
 	 * 
 	 * The editor consists of three areas, which share similar workflows as the world outliner and details panel in the level editor:
 	 * 1. Root objects,  similar to world outliner: objects (usually actors) are added here. FCreateEditorParams::ObjectSource is used to build a combo button through which new objects can be added.
-	 * 2. Subobjects, similar to component view (SSubobjectEditor): Shows subobjects of a root objects selected above; typically components.
+	 * 2. Subobjects (optional), similar to component view (SSubobjectEditor): Shows subobjects of a root objects selected above; typically components.
 	 * 3. Properties, similar to details panel: Shows properties of the selected root object and / or subobjects.
 	 */
 	CONCERTCLIENTSHAREDSLATE_API TSharedRef<IReplicationEditorView> CreateEditor(FCreateEditorParams Params);
 
+	/** Creates a default IReplicationEditorView view to use in the Unreal Editor (as opposed to on the server, etc.). */
+	inline TSharedRef<IReplicationEditorView> CreateEditorForUnrealEditor(FCreateEditorParams BaseParams)
+	{
+		BaseParams.SubobjectView = CreateUnrealEditorSubobjectView();
+		return CreateEditor(MoveTemp(BaseParams));
+	}
+	
 	/**
 	 * Creates a model that can be passed to CreateEditor.
 	 * 
