Make replication editor more modular and configurable by external modules

diff --git a/Engine/Plugins/Developer/Concert/ConcertApp/MultiUserClient/Source/MultiUserClient/Private/Replication/Client/ReplicationClient.h b/Engine/Plugins/Developer/Concert/ConcertApp/MultiUserClient/Source/MultiUserClient/Private/Replication/Client/ReplicationClient.h
--- a/Engine/Plugins/Developer/Concert/ConcertApp/MultiUserClient/Source/MultiUserClient/Private/Replication/Client/ReplicationClient.h
+++ b/Engine/Plugins/Developer/Concert/ConcertApp/MultiUserClient/Source/MultiUserClient/Private/Replication/Client/ReplicationClient.h
@@ -38,7 +38,7 @@ namespace UE::MultiUserClient
 		 * Called when the data underlying the model has changed externally. Since the change was not caused by the model,
 		 * its events, like IEditableObjectToPropertiesModel::OnObjectsChanged, were not called.
 		 * 
-		 * Subscribers are intended to call IReplicationEditorView::Refresh() in response.
+		 * Subscribers are intended to call IReplicationStreamEditor::Refresh() in response.
 		 * 
 		 * Examples: Remote client changed their stream, remote client joined with streams,
 		 * or the local client reverted a change that was rejected by the server.
diff --git a/Engine/Plugins/Developer/Concert/ConcertApp/MultiUserClient/Source/MultiUserClient/Private/Widgets/ActiveSession/Replication/Joined/SReplicationClientView.cpp b/Engine/Plugins/Developer/Concert/ConcertApp/MultiUserClient/Source/MultiUserClient/Private/Widgets/ActiveSession/Replication/Joined/SReplicationClientView.cpp
--- a/Engine/Plugins/Developer/Concert/ConcertApp/MultiUserClient/Source/MultiUserClient/Private/Widgets/ActiveSession/Replication/Joined/SReplicationClientView.cpp
+++ b/Engine/Plugins/Developer/Concert/ConcertApp/MultiUserClient/Source/MultiUserClient/Private/Widgets/ActiveSession/Replication/Joined/SReplicationClientView.cpp
@@ -5,7 +5,7 @@
 #include "Replication/Client/ReplicationClient.h"
 #include "Replication/Editor/Model/Object/EditorObjectSelectionSourceModel.h"
 #include "Replication/Editor/Model/Property/SelectPropertyFromUClassModel.h"
-#include "Replication/Editor/View/IReplicationEditorView.h"
+#include "Replication/Editor/View/IReplicationStreamEditor.h"
 #include "Replication/ReplicationWidgetFactories.h"
 #include "Widgets/ActiveSession/Replication/Joined/SReplicationClientViewToolbar.h"
 
@@ -27,7 +27,7 @@ namespace UE::MultiUserClient
 			MakeShared<FSelectPropertyFromUClassModel>()
 		};
 
-		EditorView = CreateEditorForUnrealEditor(ReplicationEditorCreationParams);
+		EditorView = CreateDefaultStreamEditor(ReplicationEditorCreationParams);
 		ChildSlot
 		[
 			SNew(SVerticalBox)
@@ -58,4 +58,4 @@ namespace UE::MultiUserClient
 	{
 		EditorView->Refresh();
 	}
-}
+}
\ No newline at end of file
diff --git a/Engine/Plugins/Developer/Concert/ConcertApp/MultiUserClient/Source/MultiUserClient/Private/Widgets/ActiveSession/Replication/Joined/SReplicationClientView.h b/Engine/Plugins/Developer/Concert/ConcertApp/MultiUserClient/Source/MultiUserClient/Private/Widgets/ActiveSession/Replication/Joined/SReplicationClientView.h
--- a/Engine/Plugins/Developer/Concert/ConcertApp/MultiUserClient/Source/MultiUserClient/Private/Widgets/ActiveSession/Replication/Joined/SReplicationClientView.h
+++ b/Engine/Plugins/Developer/Concert/ConcertApp/MultiUserClient/Source/MultiUserClient/Private/Widgets/ActiveSession/Replication/Joined/SReplicationClientView.h
@@ -7,7 +7,7 @@
 
 namespace UE::ConcertClientSharedSlate
 {
-	class IReplicationEditorView;
+	class IReplicationStreamEditor;
 }
 
 namespace UE::MultiUserClient
@@ -32,7 +32,7 @@ namespace UE::MultiUserClient
 	private:
 		
 		/** The editor view of the replication content. */
-		TSharedPtr<ConcertClientSharedSlate::IReplicationEditorView> EditorView;
+		TSharedPtr<ConcertClientSharedSlate::IReplicationStreamEditor> EditorView;
 		
 		/** The client to depict. Should always return true. If the client is destroyed, so should this widget be. */
 		TAttribute<FReplicationClient*> GetReplicationClientAttribute;
diff --git a/Engine/Plugins/Developer/Concert/ConcertApp/MultiUserClient/Source/MultiUserReplicationEditor/Private/StreamEditor/ReplicationSessionPresetEditorToolkit.cpp b/Engine/Plugins/Developer/Concert/ConcertApp/MultiUserClient/Source/MultiUserReplicationEditor/Private/StreamEditor/ReplicationSessionPresetEditorToolkit.cpp
--- a/Engine/Plugins/Developer/Concert/ConcertApp/MultiUserClient/Source/MultiUserReplicationEditor/Private/StreamEditor/ReplicationSessionPresetEditorToolkit.cpp
+++ b/Engine/Plugins/Developer/Concert/ConcertApp/MultiUserClient/Source/MultiUserReplicationEditor/Private/StreamEditor/ReplicationSessionPresetEditorToolkit.cpp
@@ -6,7 +6,7 @@
 #include "Assets/MultiUserReplicationSessionPreset.h"
 #include "Replication/Editor/Model/Property/SelectPropertyFromUClassModel.h"
 #include "Replication/Editor/Model/Object/EditorObjectSelectionSourceModel.h"
-#include "Replication/Editor/View/IReplicationEditorView.h"
+#include "Replication/Editor/View/IReplicationStreamEditor.h"
 #include "Replication/ReplicationWidgetFactories.h"
 
 #include "Framework/Docking/TabManager.h"
@@ -84,7 +84,7 @@ namespace UE::MultiUserReplicationEditor
 			}));
 		const TSharedRef<FEditorObjectSelectionSourceModel> ObjectSourceModel = MakeShared<FEditorObjectSelectionSourceModel>();
 		const TSharedRef<FSelectPropertyFromUClassModel> PropertySourceModel = MakeShared<FSelectPropertyFromUClassModel>();
-		const TSharedRef<IReplicationEditorView> EditorView = CreateEditorForUnrealEditor(FCreateEditorParams{ AssetReadWriteModel, ObjectSourceModel, PropertySourceModel });
+		const TSharedRef<IReplicationStreamEditor> EditorView = CreateDefaultStreamEditor(FCreateEditorParams{ AssetReadWriteModel, ObjectSourceModel, PropertySourceModel });
 		return SNew(SDockTab)
 			.Label(LOCTEXT("BaseDetailsTitle", "Details"))
 			[
diff --git a/Engine/Plugins/Developer/Concert/ConcertUI/ConcertClientSharedSlate/Source/ConcertClientSharedSlate/Private/Replication/Editor/Model/ReplicatedObjectData.h b/Engine/Plugins/Developer/Concert/ConcertUI/ConcertClientSharedSlate/Source/ConcertClientSharedSlate/Private/Replication/Editor/Model/ReplicatedObjectData.h
deleted file mode 100644
diff --git a/Engine/Plugins/Developer/Concert/ConcertUI/ConcertClientSharedSlate/Source/ConcertClientSharedSlate/Private/Replication/Editor/Model/ReplicatedPropertyData.h b/Engine/Plugins/Developer/Concert/ConcertUI/ConcertClientSharedSlate/Source/ConcertClientSharedSlate/Private/Replication/Editor/Model/ReplicatedPropertyData.h
deleted file mode 100644
diff --git a/Engine/Plugins/Developer/Concert/ConcertUI/ConcertClientSharedSlate/Source/ConcertClientSharedSlate/Private/Replication/Editor/View/ObjectEditor/SBaseReplicationStreamEditor.cpp b/Engine/Plugins/Developer/Concert/ConcertUI/ConcertClientSharedSlate/Source/ConcertClientSharedSlate/Private/Replication/Editor/View/ObjectEditor/SBaseReplicationStreamEditor.cpp
new file mode 100644
--- /dev/null
+++ b/Engine/Plugins/Developer/Concert/ConcertUI/ConcertClientSharedSlate/Source/ConcertClientSharedSlate/Private/Replication/Editor/View/ObjectEditor/SBaseReplicationStreamEditor.cpp
@@ -0,0 +1,230 @@
+// Copyright Epic Games, Inc. All Rights Reserved.
+
+#include "SBaseReplicationStreamEditor.h"
+
+#include "FakeObjectToPropertiesEditorModel.h"
+#include "Model/Item/SourceModelBuilders.h"
+#include "Replication/Data/ConcertPropertySelection.h"
+#include "Replication/Editor/Model/DisplayUtils.h"
+#include "Replication/Editor/Model/IEditableObjectToPropertiesModel.h"
+#include "Replication/Editor/Model/Object/IObjectSelectionSourceModel.h"
+#include "Replication/Editor/Model/Property/IPropertySelectionSourceModel.h"
+#include "Replication/Editor/Model/ReplicatedPropertyData.h"
+#include "Replication/Editor/Model/ReplicatedObjectData.h"
+#include "Replication/Editor/View/ObjectViewer/SReplicationStreamViewer.h"
+
+#include "Algo/AnyOf.h"
+#include "GameFramework/Actor.h"
+#include "UObject/Class.h"
+#include "Widgets/SBoxPanel.h"
+
+#define LOCTEXT_NAMESPACE "SObjectToPropertyEditor"
+
+namespace UE::ConcertClientSharedSlate
+{
+	void SBaseReplicationStreamEditor::Construct(
+		const FArguments& InArgs,
+		TSharedRef<IEditableObjectToPropertiesModel> InPropertiesModel,
+		TSharedRef<IObjectSelectionSourceModel> InObjectSelectionSource,
+		TSharedRef<IPropertySelectionSourceModel> InPropertySelectionSource)
+	{
+		ObjectSelectionSource = MoveTemp(InObjectSelectionSource);
+		PropertySelectionSource = MoveTemp(InPropertySelectionSource);
+		
+		EditablePropertiesModel = MoveTemp(InPropertiesModel);
+		EditablePropertiesModel->OnObjectsChanged().AddSP(this, &SBaseReplicationStreamEditor::OnObjectsChanged);
+		EditablePropertiesModel->OnPropertiesChanged().AddSP(this, &SBaseReplicationStreamEditor::OnPropertiesChanged);
+		PropertiesModelAdapter = MakeShared<FFakeObjectToPropertiesEditorModel>(EditablePropertiesModel.ToSharedRef(), PropertySelectionSource.ToSharedRef());
+
+		OnExtendObjectsContextMenuDelegate = InArgs._OnExtendObjectsContextMenu;
+		
+		ChildSlot
+		[
+			SAssignNew(ReplicationViewer, SReplicationStreamViewer, PropertiesModelAdapter.ToSharedRef())
+				.AdditionalObjectColumns(InArgs._AdditionalObjectColumns)
+				.AdditionalPropertyColumns(InArgs._AdditionalPropertyColumns)
+				.SubobjectView(InArgs._SubobjectView)
+				.OnDeleteObjects(this, &SBaseReplicationStreamEditor::OnDeleteObjects)
+				.OnObjectsContextMenuOpening(this, &SBaseReplicationStreamEditor::OnObjectsContextMenuOpening)
+				.SortPropertyRowPredicate(InArgs._SortPropertyRowPredicate)
+				.LeftOfObjectSearchBar()
+				[
+					SNew(SHorizontalBox)
+					+SHorizontalBox::Slot()
+					.AutoWidth()
+					[
+						BuildRootAddObjectWidgets()
+					]
+					+SHorizontalBox::Slot()
+					.AutoWidth()
+					[
+						InArgs._LeftOfObjectSearchBar.Widget
+					]
+				]
+				.LeftOfPropertySearchBar()
+				[
+					InArgs._LeftOfPropertySearchBar.Widget
+				]
+		];
+	}
+
+	void SBaseReplicationStreamEditor::Refresh()
+	{
+		ReplicationViewer->RefreshObjectData();
+		ReplicationViewer->RefreshSubobjectData();
+		ReplicationViewer->RefreshPropertyData();
+	}
+
+	TArray<FSoftObjectPath> SBaseReplicationStreamEditor::GetSelectedTopLevelObjects() const
+	{
+		return ReplicationViewer->GetSelectedTopLevelObjects();
+	}
+
+	TArray<FSoftObjectPath> SBaseReplicationStreamEditor::GetObjectsBeingPropertyEdited() const
+	{
+		return ReplicationViewer->GetObjectsBeingPropertyEdited();
+	}
+
+	void SBaseReplicationStreamEditor::OnObjectsChanged(TConstArrayView<UObject*> AddedObjects, TConstArrayView<FSoftObjectPath> RemovedObjects, EReplicatedObjectChangeReason ChangeReason)
+	{
+		ReplicationViewer->RefreshObjectData();
+	}
+
+	void SBaseReplicationStreamEditor::OnPropertiesChanged()
+	{
+		ReplicationViewer->RefreshPropertyData();
+	}
+
+	TSharedRef<SWidget> SBaseReplicationStreamEditor::BuildRootAddObjectWidgets() const
+	{
+		using namespace ConcertSharedSlate;
+		
+		const TSharedRef<SHorizontalBox> Root = SNew(SHorizontalBox);
+		const FSourceModelBuilders<FSelectableObjectInfo>::FItemPickerArgs Args = MakeObjectSourceBuilderArgs();
+		for (const TSourceSelectionCategory<FSelectableObjectInfo>& Category : ObjectSelectionSource->GetRootSources())
+		{
+			Root->AddSlot()
+				.AutoWidth()
+				[
+					FSourceModelBuilders<FSelectableObjectInfo>::BuildCategory(Category, Args)	
+				];
+		}
+		return Root;
+	}
+
+	void SBaseReplicationStreamEditor::OnObjectsSelectedForAdding(TArray<FSelectableObjectInfo> ObjectsToAdd) const
+	{
+		TArray<UObject*> Objects;
+		Algo::TransformIf(
+			ObjectsToAdd,
+			Objects,
+			[](const FSelectableObjectInfo& SelectableObject) { return SelectableObject.Object.IsValid(); },
+			[](const FSelectableObjectInfo& SelectableObject){ return SelectableObject.Object.Get(); }
+			);
+		
+		EditablePropertiesModel->AddObjects(Objects);
+	}
+
+	void SBaseReplicationStreamEditor::OnDeleteObjects(const TArray<TSharedPtr<FReplicatedObjectData>>& ObjectsToDelete) const
+	{
+		TArray<FSoftObjectPath> DeleteObjectPaths;
+		Algo::Transform(ObjectsToDelete, DeleteObjectPaths, [](const TSharedPtr<FReplicatedObjectData>& Data) { return Data->GetObjectPath(); });
+
+		// We want to delete children not listed in the outliner, such as components and other subobjects
+		TArray<FSoftObjectPath> ObjectAndChildren;
+		EditablePropertiesModel->ForEachReplicatedObject([&ObjectsToDelete, &ObjectAndChildren](const FSoftObjectPath& ReplicatedObject)
+		{
+			const bool bIsChildOfDeletedObject = Algo::AnyOf(ObjectsToDelete, [&ReplicatedObject](const TSharedPtr<FReplicatedObjectData>& ObjectData)
+			{
+				// ReplicatedObject is a child of a deleted object if the path before it contains one of the deleted objects
+				return ReplicatedObject.ToString().Contains(ObjectData->GetObjectPath().ToString());
+			});
+			if (bIsChildOfDeletedObject)
+			{
+				ObjectAndChildren.Add(ReplicatedObject);
+			}
+			return EBreakBehavior::Continue;
+		});
+		
+		checkf(ObjectAndChildren.Num() >= ObjectsToDelete.Num(), TEXT("Above algorithm is broken."));
+		EditablePropertiesModel->RemoveObjects(ObjectAndChildren);
+	}
+
+	TSharedPtr<SWidget> SBaseReplicationStreamEditor::OnObjectsContextMenuOpening() const
+	{
+		FMenuBuilder MenuBuilder(true, nullptr);
+		AddObjectSourceContextMenuOptions(MenuBuilder);
+		MenuBuilder.AddMenuEntry(
+				LOCTEXT("DeleteItems", "Delete"),
+				FText::GetEmpty(),
+				FSlateIcon(),
+				FUIAction(FExecuteAction::CreateSP(this, &SBaseReplicationStreamEditor::OnDeleteObjects_PassByValue, ReplicationViewer->GetSelectedOutlinerObjects())),
+				NAME_None,
+				EUserInterfaceActionType::Button
+			);
+		
+		OnExtendObjectsContextMenuDelegate.ExecuteIfBound(MenuBuilder);
+		return MenuBuilder.MakeWidget();
+	}
+
+	void SBaseReplicationStreamEditor::AddObjectSourceContextMenuOptions(FMenuBuilder& MenuBuilder) const
+	{
+		using namespace ConcertSharedSlate;
+		
+		// Context menu generation is only supported for single items
+		const TArray<TSharedPtr<FReplicatedObjectData>> SelectedObjects = ReplicationViewer->GetSelectedOutlinerObjects();
+		if (SelectedObjects.Num() != 1)
+		{
+			return;
+		}
+		
+		const TArray<TSharedRef<IObjectSourceModel>> ContextMenuOptions = ObjectSelectionSource->GetContextMenuOptions(SelectedObjects[0]->GetObjectPath());
+		if (ContextMenuOptions.IsEmpty())
+		{
+			return;
+		}
+		
+		const FSourceModelBuilders<FSelectableObjectInfo>::FItemPickerArgs Args = MakeObjectSourceBuilderArgs();
+		for (const TSharedRef<IObjectSourceModel>& SourceModel : ContextMenuOptions)
+		{
+			bool bCanAddAnyObject = false; 
+			SourceModel->EnumerateSelectableItems([this, &bCanAddAnyObject](const FSelectableObjectInfo& SelectableOption) mutable
+			{
+				bCanAddAnyObject |= !EditablePropertiesModel->ContainsObjects({ SelectableOption.Object.Get() });
+				return bCanAddAnyObject ? EBreakBehavior::Break : EBreakBehavior::Continue;
+			});
+
+			// Skip showing context menu options which will not add anything new.
+			if (bCanAddAnyObject)
+			{
+				FSourceModelBuilders<FSelectableObjectInfo>::AddOptionToMenu(SourceModel, Args, MenuBuilder);
+			}
+		}
+		
+		MenuBuilder.AddSeparator();
+	}
+
+	ConcertSharedSlate::FSourceModelBuilders<FSelectableObjectInfo>::FItemPickerArgs SBaseReplicationStreamEditor::MakeObjectSourceBuilderArgs() const
+	{
+		using namespace ConcertSharedSlate;
+		
+		return FSourceModelBuilders<FSelectableObjectInfo>::FItemPickerArgs
+		{
+			FSourceModelBuilders<FSelectableObjectInfo>::FOnItemsSelected::CreateSP(this, &SBaseReplicationStreamEditor::OnObjectsSelectedForAdding),
+			FSourceModelBuilders<FSelectableObjectInfo>::FGetItemDisplayString::CreateLambda([](const FSelectableObjectInfo& Item)
+			{
+				return Item.Object.IsValid() ? DisplayUtils::GetObjectDisplayString(*Item.Object.Get()) : TEXT("");
+			}),
+			FSourceModelBuilders<FSelectableObjectInfo>::FGetItemIcon::CreateLambda([](const FSelectableObjectInfo& Item)
+			{
+				return Item.Object.IsValid() ? DisplayUtils::GetObjectIcon(*Item.Object.Get()) : FSlateIcon();
+			}),
+			FSourceModelBuilders<FSelectableObjectInfo>::FIsItemSelected::CreateLambda([this](const FSelectableObjectInfo& Item)
+			{
+				return EditablePropertiesModel->ContainsObjects({ Item.Object.Get() } );
+			})
+		}; 
+	}
+}
+
+#undef LOCTEXT_NAMESPACE
\ No newline at end of file
diff --git a/Engine/Plugins/Developer/Concert/ConcertUI/ConcertClientSharedSlate/Source/ConcertClientSharedSlate/Private/Replication/Editor/View/ObjectEditor/SBaseReplicationStreamEditor.h b/Engine/Plugins/Developer/Concert/ConcertUI/ConcertClientSharedSlate/Source/ConcertClientSharedSlate/Private/Replication/Editor/View/ObjectEditor/SBaseReplicationStreamEditor.h
new file mode 100644
--- /dev/null
+++ b/Engine/Plugins/Developer/Concert/ConcertUI/ConcertClientSharedSlate/Source/ConcertClientSharedSlate/Private/Replication/Editor/View/ObjectEditor/SBaseReplicationStreamEditor.h
@@ -0,0 +1,121 @@
+// Copyright Epic Games, Inc. All Rights Reserved.
+
+#pragma once
+
+#include "Model/Item/SourceModelBuilders.h"
+#include "Replication/Editor/View/IReplicationStreamEditor.h"
+#include "Replication/Editor/View/PredefinedReplicationColumns.h"
+#include "Replication/ReplicationWidgetDelegates.h"
+
+#include "Widgets/DeclarativeSyntaxSupport.h"
+
+class SHorizontalBox;
+struct FConcertPropertyChain;
+
+namespace UE::ConcertClientSharedSlate
+{
+	class IEditableObjectToPropertiesModel;
+	class FFakeObjectToPropertiesEditorModel;
+	class FReplicatedObjectData;
+	class FReplicatedPropertyData;
+	class IReplicationSubobjectView;
+	class IObjectSelectionSourceModel;
+	class IObjectToPropertiesModel;
+	class IPropertySelectionSourceModel;
+	class IReplicationSubobjectView;
+	class SReplicationStreamViewer;
+	
+	enum class EReplicatedObjectChangeReason : uint8;
+	struct FSelectableObjectInfo;
+	
+	/**
+	 * Extends SObjectToPropertyView with common functionality for editing (with an IEditableObjectToPropertiesModel):
+	 * - Pressing delete calls IEditableObjectToPropertiesModel::RemoveObjects
+	 * - A combo box to left of the object search bar for adding new top level objects
+	 * - Default context menu items for top-level objects, such as deleting or adding more based right-clicked object
+	 * - Wraps the IEditableObjectToPropertiesModel with FFakeObjectToPropertiesEditorModel so the SObjectToPropertyView:
+	 *		- displays only top-level objects in the outliner (i.e. excludes components)
+	 *		- displays all properties in the property view instead of just those registered in the model
+	 *
+	 * This is a base widget and is not fully functional, yet. You must inject additional UI to support adding objects.
+	 * @see SDefaultReplicationStreamEditor for an example: it adds a Replicates? checkbox to the property rows.
+	 *
+	 * This widget assumes that the IReplicationSubobjectView::GetSelectedObjects reports subobjects that are
+	 * not in the model yet. Such objects are added to the model.
+	 */
+	class CONCERTCLIENTSHAREDSLATE_API SBaseReplicationStreamEditor : public IReplicationStreamEditor
+	{
+	public:
+		
+		SLATE_BEGIN_ARGS(SBaseReplicationStreamEditor)
+		{}
+			/** Additional columns to add to the object view */
+			SLATE_ARGUMENT(TArray<ReplicationObjectColumns::FReplicationObjectColumn>, AdditionalObjectColumns)
+			/** Additional columns to add to the property view */
+			SLATE_ARGUMENT(TArray<ReplicationPropertyColumns::FReplicationPropertyColumn>, AdditionalPropertyColumns)
+		
+			/** Optional. Placed between root object outliner and property editor. */
+			SLATE_ARGUMENT(TSharedPtr<IReplicationSubobjectView>, SubobjectView)
+
+			/** Called to generate the context menu for objects. This extends the options already generated by this widget. */
+			SLATE_EVENT(FExtendMenu, OnExtendObjectsContextMenu)
+
+			/** Optional. Used for determining the order in which properties are displayed. */
+			SLATE_EVENT(FSortPropertyPredicate, SortPropertyRowPredicate)
+
+			/** Optional widget to add to the left of the object list search bar. */
+			SLATE_NAMED_SLOT(FArguments, LeftOfObjectSearchBar)
+			/** Optional widget to add to the left of the property list search bar. */
+			SLATE_NAMED_SLOT(FArguments, LeftOfPropertySearchBar)
+		SLATE_END_ARGS()
+
+		void Construct(const FArguments& InArgs,
+           TSharedRef<IEditableObjectToPropertiesModel> InPropertiesModel,
+           TSharedRef<IObjectSelectionSourceModel> InObjectSelectionSource,
+           TSharedRef<IPropertySelectionSourceModel> InPropertySelectionSource
+		);
+
+		//~ Begin IReplicationStreamEditor Interface
+		virtual void Refresh() override;
+		virtual TArray<FSoftObjectPath> GetSelectedTopLevelObjects() const override;
+		virtual TArray<FSoftObjectPath> GetObjectsBeingPropertyEdited() const override;
+		//~ End IReplicationStreamEditor Interface
+
+	private:
+
+		/** Displays the objects and properties */
+		TSharedPtr<SReplicationStreamViewer> ReplicationViewer;
+
+		/** For reading and writting to the edited asset */
+		TSharedPtr<IEditableObjectToPropertiesModel> EditablePropertiesModel;
+		/**
+		 * Fakes to SObjectToPropertyView that all UClass properties are contained.
+		 * We inject checkboxes to SObjectToPropertyView which do the actual adding and removing.
+		 */
+		TSharedPtr<FFakeObjectToPropertiesEditorModel> PropertiesModelAdapter;
+		
+		/** For deciding which objects can be added to EditablePropertiesModel. */
+		TSharedPtr<IObjectSelectionSourceModel> ObjectSelectionSource;
+		/** For deciding which properties can be added to EditablePropertiesModel. */
+		TSharedPtr<IPropertySelectionSourceModel> PropertySelectionSource;
+
+		/** Generates additional entries */
+		FExtendMenu OnExtendObjectsContextMenuDelegate;
+		
+		void OnObjectsChanged(TConstArrayView<UObject*> AddedObjects, TConstArrayView<FSoftObjectPath> RemovedObjects, EReplicatedObjectChangeReason ChangeReason);
+		void OnPropertiesChanged();
+
+		// Customizing left objects search bar
+		TSharedRef<SWidget> BuildRootAddObjectWidgets() const;
+		void OnObjectsSelectedForAdding(TArray<FSelectableObjectInfo> ObjectsToAdd) const;
+
+		// Respond to object editing events
+		void OnDeleteObjects_PassByValue(TArray<TSharedPtr<FReplicatedObjectData>> CopiedObjectsToDelete) const { OnDeleteObjects(CopiedObjectsToDelete); }
+		void OnDeleteObjects(const TArray<TSharedPtr<FReplicatedObjectData>>& ObjectsToDelete) const;
+		TSharedPtr<SWidget> OnObjectsContextMenuOpening() const;
+		void AddObjectSourceContextMenuOptions(FMenuBuilder& MenuBuilder) const;
+		
+		// Utils for building item source widgets
+		ConcertSharedSlate::FSourceModelBuilders<FSelectableObjectInfo>::FItemPickerArgs MakeObjectSourceBuilderArgs() const;
+	};
+}
diff --git a/Engine/Plugins/Developer/Concert/ConcertUI/ConcertClientSharedSlate/Source/ConcertClientSharedSlate/Private/Replication/Editor/View/ObjectEditor/SDefaultReplicationStreamEditor.cpp b/Engine/Plugins/Developer/Concert/ConcertUI/ConcertClientSharedSlate/Source/ConcertClientSharedSlate/Private/Replication/Editor/View/ObjectEditor/SDefaultReplicationStreamEditor.cpp
new file mode 100644
--- /dev/null
+++ b/Engine/Plugins/Developer/Concert/ConcertUI/ConcertClientSharedSlate/Source/ConcertClientSharedSlate/Private/Replication/Editor/View/ObjectEditor/SDefaultReplicationStreamEditor.cpp
@@ -0,0 +1,71 @@
+// Copyright Epic Games, Inc. All Rights Reserved.
+
+#include "SDefaultReplicationStreamEditor.h"
+
+#include "Replication/Editor/Model/IEditableObjectToPropertiesModel.h"
+#include "Replication/Editor/View/ObjectViewer/Tree/SelectionViewerColumns.h"
+
+#include "SBaseReplicationStreamEditor.h"
+
+namespace UE::ConcertClientSharedSlate
+{
+	void SDefaultReplicationStreamEditor::Construct(
+		const FArguments& InArgs,
+		TSharedRef<IEditableObjectToPropertiesModel> InPropertiesModel,
+		TSharedRef<IObjectSelectionSourceModel> InObjectSelectionSource,
+		TSharedRef<IPropertySelectionSourceModel> InPropertySelectionSource)
+	{
+		PropertiesModel = InPropertiesModel;
+		
+		using namespace ReplicationPropertyColumns;
+		const FReplicationPropertyColumn ReplicatesColumn = ReplicatesColumns(SharedThis(this), InPropertiesModel);
+		TArray<FReplicationPropertyColumn> PropertyColumns = InArgs._AdditionalPropertyColumns;
+		PropertyColumns.Add(ReplicatesColumn);
+		
+		ChildSlot
+		[
+			SAssignNew(WrappedEditor, SBaseReplicationStreamEditor, InPropertiesModel, InObjectSelectionSource, InPropertySelectionSource)
+				.AdditionalObjectColumns(InArgs._AdditionalObjectColumns)
+				.AdditionalPropertyColumns(PropertyColumns)
+				.SubobjectView(InArgs._SubobjectView)
+				.OnExtendObjectsContextMenu(InArgs._OnExtendObjectsContextMenu)
+				.SortPropertyRowPredicate(this, &SDefaultReplicationStreamEditor::SortPropertiesPredicate)
+				.LeftOfObjectSearchBar()
+				[
+					InArgs._LeftOfObjectSearchBar.Widget
+				]
+				.LeftOfPropertySearchBar()
+				[
+					InArgs._LeftOfPropertySearchBar.Widget
+				]
+		];
+	}
+
+	void SDefaultReplicationStreamEditor::Refresh()
+	{
+		return WrappedEditor->Refresh();
+	}
+
+	TArray<FSoftObjectPath> SDefaultReplicationStreamEditor::GetSelectedTopLevelObjects() const
+	{
+		return WrappedEditor->GetSelectedTopLevelObjects();
+	}
+
+	TArray<FSoftObjectPath> SDefaultReplicationStreamEditor::GetObjectsBeingPropertyEdited() const
+	{
+		return WrappedEditor->GetObjectsBeingPropertyEdited();
+	}
+
+	bool SDefaultReplicationStreamEditor::SortPropertiesPredicate(
+		const FReplicatedPropertyData& Left,
+		const FReplicatedPropertyData& Right
+		) const
+	{
+		return ReplicationPropertyColumns::SortBySelectionThenByName_PropertyPredicate(
+			WrappedEditor->GetObjectsBeingPropertyEdited(),
+			*PropertiesModel,
+			Left,
+			Right
+			);
+	}
+}
diff --git a/Engine/Plugins/Developer/Concert/ConcertUI/ConcertClientSharedSlate/Source/ConcertClientSharedSlate/Private/Replication/Editor/View/ObjectEditor/SDefaultReplicationStreamEditor.h b/Engine/Plugins/Developer/Concert/ConcertUI/ConcertClientSharedSlate/Source/ConcertClientSharedSlate/Private/Replication/Editor/View/ObjectEditor/SDefaultReplicationStreamEditor.h
new file mode 100644
--- /dev/null
+++ b/Engine/Plugins/Developer/Concert/ConcertUI/ConcertClientSharedSlate/Source/ConcertClientSharedSlate/Private/Replication/Editor/View/ObjectEditor/SDefaultReplicationStreamEditor.h
@@ -0,0 +1,78 @@
+// Copyright Epic Games, Inc. All Rights Reserved.
+
+#pragma once
+
+#include "Replication/Editor/View/IReplicationStreamEditor.h"
+#include "Replication/Editor/View/PredefinedReplicationColumns.h"
+#include "Replication/ReplicationWidgetDelegates.h"
+
+#include "Widgets/DeclarativeSyntaxSupport.h"
+
+class SHorizontalBox;
+struct FConcertPropertyChain;
+
+namespace UE::ConcertClientSharedSlate
+{
+	class IEditableObjectToPropertiesModel;
+	class IObjectSelectionSourceModel;
+	class IPropertySelectionSourceModel;
+	class IReplicationSubobjectView;
+	class SBaseReplicationStreamEditor;
+
+	/**
+	 * The default replication stream editor adds a checkbox in front of every property row.
+	 * 
+	 * The checkbox allows adding / removing the properties to / from the selected objects.
+	 * The properties are sorted based on the check state: checked come before unchecked.
+	 */
+	class SDefaultReplicationStreamEditor : public IReplicationStreamEditor
+    {
+    public:
+		
+		SLATE_BEGIN_ARGS(SDefaultReplicationStreamEditor)
+		{}
+			/** Additional columns to add to the object view */
+			SLATE_ARGUMENT(TArray<ReplicationObjectColumns::FReplicationObjectColumn>, AdditionalObjectColumns)
+			/** Additional columns to add to the property view */
+			SLATE_ARGUMENT(TArray<ReplicationPropertyColumns::FReplicationPropertyColumn>, AdditionalPropertyColumns)
+			
+			/** Optional. Placed between root object outliner and property editor. */
+			SLATE_ARGUMENT(TSharedPtr<IReplicationSubobjectView>, SubobjectView)
+
+			/** Called to generate the context menu for objects. This extends the options already generated by this widget. */
+			SLATE_EVENT(FExtendMenu, OnExtendObjectsContextMenu)
+
+			/** Optional. Used for determining the order in which properties are displayed. */
+			SLATE_EVENT(FSortPropertyPredicate, SortPropertyRowPredicate)
+
+			/** Optional widget to add to the left of the object list search bar. */
+			SLATE_NAMED_SLOT(FArguments, LeftOfObjectSearchBar)
+			/** Optional widget to add to the left of the property list search bar. */
+			SLATE_NAMED_SLOT(FArguments, LeftOfPropertySearchBar)
+		SLATE_END_ARGS()
+
+		void Construct(
+			const FArguments& InArgs,
+			TSharedRef<IEditableObjectToPropertiesModel> InPropertiesModel,
+			TSharedRef<IObjectSelectionSourceModel> InObjectSelectionSource,
+			TSharedRef<IPropertySelectionSourceModel> InPropertySelectionSource
+			);
+
+		//~ Begin IReplicationStreamEditor Interface
+		virtual void Refresh() override;
+		virtual TArray<FSoftObjectPath> GetSelectedTopLevelObjects() const override;
+		virtual TArray<FSoftObjectPath> GetObjectsBeingPropertyEdited() const override;
+		//~ End IReplicationStreamEditor Interface
+
+	private:
+
+		/** Core editor implementing generic editing UI & workflows. */
+		TSharedPtr<SBaseReplicationStreamEditor> WrappedEditor;
+		
+		/** Model needed for sorting properties based on selection. */
+		TSharedPtr<IEditableObjectToPropertiesModel> PropertiesModel;
+		
+		bool SortPropertiesPredicate(const FReplicatedPropertyData& Left, const FReplicatedPropertyData& Right) const;
+    };
+}
+
diff --git a/Engine/Plugins/Developer/Concert/ConcertUI/ConcertClientSharedSlate/Source/ConcertClientSharedSlate/Private/Replication/Editor/View/ObjectEditor/SObjectToPropertyEditor.cpp b/Engine/Plugins/Developer/Concert/ConcertUI/ConcertClientSharedSlate/Source/ConcertClientSharedSlate/Private/Replication/Editor/View/ObjectEditor/SObjectToPropertyEditor.cpp
deleted file mode 100644
diff --git a/Engine/Plugins/Developer/Concert/ConcertUI/ConcertClientSharedSlate/Source/ConcertClientSharedSlate/Private/Replication/Editor/View/ObjectEditor/SObjectToPropertyEditor.h b/Engine/Plugins/Developer/Concert/ConcertUI/ConcertClientSharedSlate/Source/ConcertClientSharedSlate/Private/Replication/Editor/View/ObjectEditor/SObjectToPropertyEditor.h
deleted file mode 100644
diff --git a/Engine/Plugins/Developer/Concert/ConcertUI/ConcertClientSharedSlate/Source/ConcertClientSharedSlate/Private/Replication/Editor/View/ObjectViewer/Property/SReplicatedPropertiesView.cpp b/Engine/Plugins/Developer/Concert/ConcertUI/ConcertClientSharedSlate/Source/ConcertClientSharedSlate/Private/Replication/Editor/View/ObjectViewer/Property/SReplicatedPropertiesView.cpp
--- a/Engine/Plugins/Developer/Concert/ConcertUI/ConcertClientSharedSlate/Source/ConcertClientSharedSlate/Private/Replication/Editor/View/ObjectViewer/Property/SReplicatedPropertiesView.cpp
+++ b/Engine/Plugins/Developer/Concert/ConcertUI/ConcertClientSharedSlate/Source/ConcertClientSharedSlate/Private/Replication/Editor/View/ObjectViewer/Property/SReplicatedPropertiesView.cpp
@@ -19,7 +19,7 @@ namespace UE::ConcertClientSharedSlate
 		
 		ChildSlot
 		[
-			SAssignNew(ReplicatedProperties, SReplicationTreeView<TSharedPtr<FReplicatedPropertyData>>)
+			SAssignNew(ReplicatedProperties, SReplicationTreeView<FReplicatedPropertyData>)
 				.RootItemsSource(InArgs._RootItemsSource)
 				.OnGetChildren(InArgs._OnGetChildren)
 				.Columns(InArgs._Columns)
diff --git a/Engine/Plugins/Developer/Concert/ConcertUI/ConcertClientSharedSlate/Source/ConcertClientSharedSlate/Private/Replication/Editor/View/ObjectViewer/Property/SReplicatedPropertiesView.h b/Engine/Plugins/Developer/Concert/ConcertUI/ConcertClientSharedSlate/Source/ConcertClientSharedSlate/Private/Replication/Editor/View/ObjectViewer/Property/SReplicatedPropertiesView.h
--- a/Engine/Plugins/Developer/Concert/ConcertUI/ConcertClientSharedSlate/Source/ConcertClientSharedSlate/Private/Replication/Editor/View/ObjectViewer/Property/SReplicatedPropertiesView.h
+++ b/Engine/Plugins/Developer/Concert/ConcertUI/ConcertClientSharedSlate/Source/ConcertClientSharedSlate/Private/Replication/Editor/View/ObjectViewer/Property/SReplicatedPropertiesView.h
@@ -25,10 +25,10 @@ namespace UE::ConcertClientSharedSlate
 			SLATE_ARGUMENT(TArray<TSharedPtr<FReplicatedPropertyData>>*, RootItemsSource)
 		
 			/** Gets an items children for the tree view */
-			SLATE_EVENT(SReplicationTreeView<TSharedPtr<FReplicatedPropertyData>>::FGetItemChildren, OnGetChildren)
+			SLATE_EVENT(SReplicationTreeView<FReplicatedPropertyData>::FGetItemChildren, OnGetChildren)
 		
 			/** The columns this list should have */
-			SLATE_ARGUMENT(TArray<TReplicationColumn<TSharedPtr<FReplicatedPropertyData>>>, Columns)
+			SLATE_ARGUMENT(TArray<TReplicationColumn<FReplicatedPropertyData>>, Columns)
 			/** The name of the column that will have the SExpanderArrow for the tree view. */
 			SLATE_ARGUMENT(FName, ExpandableColumnLabel)
 		
@@ -50,7 +50,7 @@ namespace UE::ConcertClientSharedSlate
 	private:
 
 		/** The tree view displaying the replicated properties */
-		TSharedPtr<SReplicationTreeView<TSharedPtr<FReplicatedPropertyData>>> ReplicatedProperties;
+		TSharedPtr<SReplicationTreeView<FReplicatedPropertyData>> ReplicatedProperties;
 		/** Displays the active filters*/
 		TSharedPtr<SBasicFilterBar<TSharedPtr<FReplicatedPropertyData>>> FilterBar;
 
diff --git a/Engine/Plugins/Developer/Concert/ConcertUI/ConcertClientSharedSlate/Source/ConcertClientSharedSlate/Private/Replication/Editor/View/ObjectViewer/SObjectToPropertyView.cpp b/Engine/Plugins/Developer/Concert/ConcertUI/ConcertClientSharedSlate/Source/ConcertClientSharedSlate/Private/Replication/Editor/View/ObjectViewer/SObjectToPropertyView.cpp
deleted file mode 100644
diff --git a/Engine/Plugins/Developer/Concert/ConcertUI/ConcertClientSharedSlate/Source/ConcertClientSharedSlate/Private/Replication/Editor/View/ObjectViewer/SObjectToPropertyView.h b/Engine/Plugins/Developer/Concert/ConcertUI/ConcertClientSharedSlate/Source/ConcertClientSharedSlate/Private/Replication/Editor/View/ObjectViewer/SObjectToPropertyView.h
deleted file mode 100644
diff --git a/Engine/Plugins/Developer/Concert/ConcertUI/ConcertClientSharedSlate/Source/ConcertClientSharedSlate/Private/Replication/Editor/View/ObjectViewer/SReplicationStreamViewer.cpp b/Engine/Plugins/Developer/Concert/ConcertUI/ConcertClientSharedSlate/Source/ConcertClientSharedSlate/Private/Replication/Editor/View/ObjectViewer/SReplicationStreamViewer.cpp
new file mode 100644
--- /dev/null
+++ b/Engine/Plugins/Developer/Concert/ConcertUI/ConcertClientSharedSlate/Source/ConcertClientSharedSlate/Private/Replication/Editor/View/ObjectViewer/SReplicationStreamViewer.cpp
@@ -0,0 +1,237 @@
+// Copyright Epic Games, Inc. All Rights Reserved.
+
+#include "SReplicationStreamViewer.h"
+
+#include "ConcertFrontendUtils.h"
+#include "Replication/Editor/Model/IObjectToPropertiesModel.h"
+#include "Replication/Editor/Model/ReplicatedObjectData.h"
+#include "Replication/Editor/Model/ReplicatedPropertyData.h"
+#include "Replication/Editor/View/ObjectViewer/Property/SReplicatedPropertiesView.h"
+#include "Replication/Editor/View/ObjectViewer/Tree/SelectionViewerColumns.h"
+#include "SSubobjectAndPropertySection.h"
+
+#include "Widgets/Layout/SBorder.h"
+#include "Widgets/Layout/SSplitter.h"
+#include "Widgets/SBoxPanel.h"
+#include "Widgets/Text/STextBlock.h"
+
+#define LOCTEXT_NAMESPACE "SObjectToPropertyView"
+
+namespace UE::ConcertClientSharedSlate
+{
+	void SReplicationStreamViewer::Construct(const FArguments& InArgs, TSharedRef<IObjectToPropertiesModel> InPropertiesModel)
+	{
+		PropertiesModel = MoveTemp(InPropertiesModel);
+		
+		ChildSlot
+		[
+			CreateContentWidget(InArgs)
+		];
+
+		RefreshObjectData();
+		RefreshSubobjectData();
+		RefreshPropertyData();
+
+		PropertyArea->SetExpanded(true);
+	}
+
+	void SReplicationStreamViewer::Refresh()
+	{
+		RefreshObjectData();
+		RefreshPropertyData();
+		RefreshSubobjectData();
+	}
+
+	TArray<FSoftObjectPath> SReplicationStreamViewer::GetSelectedTopLevelObjects() const
+	{
+		TArray<FSoftObjectPath> Result;
+		Algo::Transform(GetSelectedOutlinerObjects(), Result, [](const TSharedPtr<FReplicatedObjectData>& Data)
+		{
+			return Data->GetObjectPath();
+		});
+		return Result;
+	}
+
+	TArray<FSoftObjectPath> SReplicationStreamViewer::GetObjectsBeingPropertyEdited() const
+	{
+		return SubobjectAndPropertySection->GetSelectedObjects();
+	}
+
+	void SReplicationStreamViewer::RefreshObjectData()
+	{
+		const int32 NumElements = PropertiesModel->GetNumReplicatedObjects();
+		// Re-using existing instances is tricky: we cannot update the object path in an item because the list view will no detect this change;
+		// list view only looks at the shared ptr address. So the UI will not be refreshed. Since the number of items will be small, just reallocate... 
+		ObjectRowData.Empty(NumElements);
+
+		// Try to re-use old instances by using the old PathToObjectDataCache. This is also done so the expansion states restore correctly in the tree view.
+		TMap<FSoftObjectPath, TSharedPtr<FReplicatedObjectData>> NewPathToObjectDataCache;
+		
+		// Do a complete refresh.
+		// Complete refresh is acceptable because the list is updated infrequently and typically small < 500 items.
+		// An alternative would be to change RefreshObjectData to be called with two variables ObjectsAdded and ObjectsRemoved.
+		PropertiesModel->ForEachReplicatedObject([this, &NewPathToObjectDataCache](const FSoftObjectPath& Path) mutable
+		{
+			const TSharedPtr<FReplicatedObjectData>* ExistingItem = PathToObjectDataCache.Find(Path);
+			const TSharedRef<FReplicatedObjectData> Item = ExistingItem ? ExistingItem->ToSharedRef() : AllocateObjectData(Path);
+			ObjectRowData.Emplace(Item);
+			NewPathToObjectDataCache.Emplace(Path, Item);
+			return EBreakBehavior::Continue;
+		});
+
+		// Only refresh the tree if it is necessary as it causes us to select stuff in the subobject view
+		if (!PathToObjectDataCache.OrderIndependentCompareEqual(NewPathToObjectDataCache))
+		{
+			// If an item was removed, then NewPathToObjectDataCache does not contain it. 
+			PathToObjectDataCache = MoveTemp(NewPathToObjectDataCache);
+
+			// The tree view requires the item source to only contain the root items. Children are discovered via GetObjectRowChildren. We re-use GetObjectRowChildren to remove any non-root nodes.
+			BuildRootObjectRowData();
+			ReplicatedObjects->OnItemsChanged();
+		}
+	}
+
+	void SReplicationStreamViewer::RefreshSubobjectData()
+	{
+		SubobjectAndPropertySection->RefreshSubobjectData();
+		
+		// Need to update properties in case the subobject view changed selection
+		RefreshPropertyData();
+	}
+
+	void SReplicationStreamViewer::RefreshPropertyData()
+	{
+		SubobjectAndPropertySection->RefreshPropertyData();
+	}
+	
+	const TArray<TSharedPtr<FReplicatedPropertyData>>& SReplicationStreamViewer::GetPropertyRowData() const
+	{
+		return SubobjectAndPropertySection->GetPropertyRowData();
+	}
+
+	TSharedRef<FReplicatedObjectData> SReplicationStreamViewer::AllocateObjectData(FSoftObjectPath ObjectPath)
+	{
+		return MakeShared<FReplicatedObjectData>(MoveTemp(ObjectPath));
+	}
+
+	TSharedRef<SWidget> SReplicationStreamViewer::CreateContentWidget(const FArguments& InArgs)
+	{
+		return SNew(SSplitter)
+			.Orientation(Orient_Vertical)
+
+			+SSplitter::Slot()
+			 .Value(1.f)
+			[
+				CreateActorsSection(InArgs)
+			]
+
+			+SSplitter::Slot()
+			.SizeRule(TAttribute<SSplitter::ESizeRule>(this, &SReplicationStreamViewer::GetPropertyAreaSizeRule))
+			.Value(2.f)
+			[
+				CreatePropertiesSection(InArgs)
+			];
+	}
+
+	TSharedRef<SWidget> SReplicationStreamViewer::CreateActorsSection(const FArguments& InArgs)
+	{
+		TArray<ReplicationObjectColumns::FReplicationObjectColumn> Columns
+		{
+			ReplicationObjectColumns::IconColumn(PropertiesModel.ToSharedRef()),
+			ReplicationObjectColumns::LabelColumn(),
+			ReplicationObjectColumns::TypeColumn(PropertiesModel.ToSharedRef())
+		};
+		Columns.Append(InArgs._AdditionalObjectColumns);
+		
+		return SAssignNew(ReplicatedObjects, SReplicationTreeView<FReplicatedObjectData>)
+			.RootItemsSource(&RootObjectRowData)
+			.OnGetChildren(this, &SReplicationStreamViewer::GetObjectRowChildren)
+			.OnContextMenuOpening(InArgs._OnObjectsContextMenuOpening)
+			.OnDeleteItems(InArgs._OnDeleteObjects)
+			.OnSelectionChanged_Lambda([this]()
+			{
+				RefreshSubobjectData();
+				// The subobject UI should automatically highlight the root objects after they are selected in the outliner
+				SubobjectAndPropertySection->SelectRootObjects();
+			})
+			.Columns(Columns)
+			.ExpandableColumnLabel(ReplicationObjectColumns::LabelColumnId)
+			.SelectionMode(ESelectionMode::Multi)
+			.LeftOfSearchBar()
+			[
+				InArgs._LeftOfObjectSearchBar.Widget
+			];
+	}
+
+	TSharedRef<SWidget> SReplicationStreamViewer::CreatePropertiesSection(const FArguments& InArgs)
+	{
+		return SNew(SBorder)
+			.BorderImage(FAppStyle::Get().GetBrush("ToolPanel.GroupBorder"))
+			.Padding(0.0f)
+			[
+				SAssignNew(PropertyArea, SExpandableArea)
+				.InitiallyCollapsed(true) 
+				.BorderBackgroundColor(FLinearColor(0.6f, 0.6f, 0.6f, 1.0f))
+				.BorderImage_Lambda([this]() { return ConcertFrontendUtils::GetExpandableAreaBorderImage(*PropertyArea); })
+				.BodyBorderImage(FAppStyle::Get().GetBrush("ToolPanel.GroupBorder"))
+				.BodyBorderBackgroundColor(FLinearColor::White)
+				.OnAreaExpansionChanged(this, &SReplicationStreamViewer::OnPropertyAreaExpansionChanged)
+				.Padding(0.0f)
+				.HeaderContent()
+				[
+					SNew(STextBlock)
+					.Text(LOCTEXT("ReplicatedProperties", "Subobjects & Properties"))
+					.Font(FAppStyle::Get().GetFontStyle("DetailsView.CategoryFontStyle"))
+					.ShadowOffset(FVector2D(1.0f, 1.0f))
+				]
+				.BodyContent()
+				[
+					SAssignNew(SubobjectAndPropertySection, SSubobjectAndPropertySection, PropertiesModel.ToSharedRef())
+					.AdditionalPropertyColumns(InArgs._AdditionalPropertyColumns)
+					.SubobjectView(InArgs._SubobjectView)
+					.SortPropertyRowPredicate(InArgs._SortPropertyRowPredicate)
+					.GetSelectedRootObjects_Lambda([this](){ return GetSelectedOutlinerObjects(); })
+					.LeftOfPropertySearchBar()
+					[
+						InArgs._LeftOfPropertySearchBar.Widget
+					]
+				]
+			];
+	}
+
+	void SReplicationStreamViewer::BuildRootObjectRowData()
+	{
+		TSet<TSharedPtr<FReplicatedObjectData>> NonRootNodes;
+		for (const TSharedPtr<FReplicatedObjectData>& Node : ObjectRowData)
+		{
+			GetObjectRowChildren(Node, [&NonRootNodes](TSharedPtr<FReplicatedObjectData> Child)
+			{
+				NonRootNodes.Add(MoveTemp(Child));
+			});
+		}
+
+		// Make RootObjectRowData only contain those nodes which were not listed as children 
+		RootObjectRowData.Empty(NonRootNodes.Num());
+		for (const TSharedPtr<FReplicatedObjectData>& Node : ObjectRowData)
+		{
+			if (!NonRootNodes.Contains(Node))
+			{
+				RootObjectRowData.Add(Node);
+			}
+		}
+		
+		RootObjectRowData.Sort([](const TSharedPtr<FReplicatedObjectData>& Left, const TSharedPtr<FReplicatedObjectData>& Right)
+		{
+			return Left->GetObjectPath().GetSubPathString() < Right->GetObjectPath().GetSubPathString();
+		});
+	}
+
+	void SReplicationStreamViewer::GetObjectRowChildren(TSharedPtr<FReplicatedObjectData> ReplicatedObjectData, TFunctionRef<void(TSharedPtr<FReplicatedObjectData>)> ProcessChild)
+	{
+		// For now there are no children in the outliner.
+		// In the future we could add an external delegate that allows child actors to be parented.
+		// Important: this view should be possible to be built in programs, so it should not reference things like AActor, UActorComponent, ResolveObject, etc. directly.
+	}
+}
+
+#undef LOCTEXT_NAMESPACE
\ No newline at end of file
diff --git a/Engine/Plugins/Developer/Concert/ConcertUI/ConcertClientSharedSlate/Source/ConcertClientSharedSlate/Private/Replication/Editor/View/ObjectViewer/SReplicationStreamViewer.h b/Engine/Plugins/Developer/Concert/ConcertUI/ConcertClientSharedSlate/Source/ConcertClientSharedSlate/Private/Replication/Editor/View/ObjectViewer/SReplicationStreamViewer.h
new file mode 100644
--- /dev/null
+++ b/Engine/Plugins/Developer/Concert/ConcertUI/ConcertClientSharedSlate/Source/ConcertClientSharedSlate/Private/Replication/Editor/View/ObjectViewer/SReplicationStreamViewer.h
@@ -0,0 +1,132 @@
+// Copyright Epic Games, Inc. All Rights Reserved.
+
+#pragma once
+
+#include "Replication/Editor/View/IReplicationStreamViewer.h"
+
+#include "Replication/Editor/Model/ReplicatedObjectData.h"
+#include "Replication/Editor/View/ObjectViewer/Tree/SReplicationTreeView.h"
+#include "Replication/Editor/View/ObjectViewer/Tree/SelectionViewerColumns.h"
+#include "SSubobjectAndPropertySection.h"
+
+#include "Algo/Transform.h"
+#include "Misc/Optional.h"
+#include "Widgets/DeclarativeSyntaxSupport.h"
+#include "Widgets/SCompoundWidget.h"
+
+class SExpandableArea;
+class SWidgetSwitcher;
+struct FSoftClassPath;
+struct FSoftObjectPath;
+
+namespace UE::ConcertClientSharedSlate
+{
+	class SSubobjectAndPropertySection;
+	class FReplicatedPropertyData;
+	class FReplicatedObjectData;
+	class IEditableObjectToPropertiesModel;
+	class IReplicationSubobjectView;
+	class IObjectToPropertiesModel;
+	class SReplicatedPropertiesView;
+	
+	/**
+	 * Root widget for viewing UMultiUserPropertyReplicationSelection.
+	 * This widget knows how to display IObjectToPropertiesModel.
+	 * 
+	 * The underlying data is modified by SObjectToPropertyEditor, which uses this widget's extension
+	 * points to call functions on IEditableObjectToPropertiesModel.
+	 *
+	 * Important: this view should be possible to be built in programs, so it should not reference things like AActor,
+	 * UActorComponent, ResolveObject, etc. directly. 
+	 */
+	class CONCERTCLIENTSHAREDSLATE_API SReplicationStreamViewer : public IReplicationStreamViewer
+	{
+	public:
+
+		SLATE_BEGIN_ARGS(SReplicationStreamViewer)
+		{}
+			/** Additional columns to add to the object view */
+			SLATE_ARGUMENT(TArray<ReplicationObjectColumns::FReplicationObjectColumn>, AdditionalObjectColumns)
+			/** Additional columns to add to the property view */
+			SLATE_ARGUMENT(TArray<ReplicationPropertyColumns::FReplicationPropertyColumn>, AdditionalPropertyColumns)
+
+			/** Optional. Placed between root object outliner and property editor. */
+			SLATE_ARGUMENT(TSharedPtr<IReplicationSubobjectView>, SubobjectView)
+
+			/** Optional. Called when the delete key is pressed in the object view. */
+			SLATE_EVENT(SReplicationTreeView<FReplicatedObjectData>::FDeleteItems, OnDeleteObjects)
+		
+			/** Called to generate the context menu for objects. */
+			SLATE_EVENT(FOnContextMenuOpening, OnObjectsContextMenuOpening)
+
+			/** Optional. Used for determining the order in which properties are displayed. */
+			SLATE_EVENT(FSortPropertyPredicate, SortPropertyRowPredicate)
+		
+			/** Optional widget to add to the left of the object list search bar. */
+			SLATE_NAMED_SLOT(FArguments, LeftOfObjectSearchBar)
+			/** Optional widget to add to the left of the property list search bar. */
+			SLATE_NAMED_SLOT(FArguments, LeftOfPropertySearchBar)
+		SLATE_END_ARGS()
+
+		void Construct(const FArguments& InArgs, TSharedRef<IObjectToPropertiesModel> InPropertiesModel);
+
+		//~ Begin IReplicationStreamViewer Interface
+		virtual void Refresh() override;
+		virtual TArray<FSoftObjectPath> GetSelectedTopLevelObjects() const override;
+		virtual TArray<FSoftObjectPath> GetObjectsBeingPropertyEdited() const override;
+		//~ End IReplicationStreamViewer Interface
+
+		void RefreshObjectData();
+		void RefreshSubobjectData();
+		void RefreshPropertyData();
+
+		const TSet<TSharedPtr<FReplicatedObjectData>>& GetObjectRowData() const { return ObjectRowData; }
+		const TArray<TSharedPtr<FReplicatedPropertyData>>& GetPropertyRowData() const;
+
+		/** @return Gets the root objects selected in the outliner; the subobject view chooses which of these objects (or their subobjects) end up in GetSelectedObjectShowingProperties. */
+		TArray<TSharedPtr<FReplicatedObjectData>> GetSelectedOutlinerObjects() const { return ReplicatedObjects->GetSelectedItems(); }
+		
+	private:
+
+		/** The model this view is visualizing. */
+		TSharedPtr<IObjectToPropertiesModel> PropertiesModel;
+
+		/** Lists the properties of the selected actor */
+		TSharedPtr<SExpandableArea> PropertyArea;
+		/** Edits the property list and (optionally) exposes subobjects of the selected root object. */
+		TSharedPtr<SSubobjectAndPropertySection> SubobjectAndPropertySection;
+
+		/** Tree view for replicated objects. */
+		TSharedPtr<SReplicationTreeView<FReplicatedObjectData>> ReplicatedObjects;
+		
+		/**
+		 * These instances can be subclasses of FReplicatedObjectData, e.g. FReplicatedObjectData_Editor.
+		 * Their type can be overridden by subclasses.
+		 * They only have the FReplicatedObjectData type so they can be passed efficiently to SObjectToPropertyView.
+		 * @see GetObjectData
+		 */
+		TSet<TSharedPtr<FReplicatedObjectData>> ObjectRowData;
+		/** The instances of ObjectRowData which do not have any parents. This acts as the item source for the tree view. */
+		TArray<TSharedPtr<FReplicatedObjectData>> RootObjectRowData;
+		/** Inverse map of ObjectRowData using FReplicatedObjectData::GetObjectPath as key. Contains all elements of ObjectRowData. */
+		TMap<FSoftObjectPath, TSharedPtr<FReplicatedObjectData>> PathToObjectDataCache;
+
+		bool bIsPropertyAreaExpanded = false;
+
+		TSharedRef<FReplicatedObjectData> AllocateObjectData(FSoftObjectPath ObjectPath);
+
+		// Widget creation helpers
+		TSharedRef<SWidget> CreateContentWidget(const FArguments& InArgs);
+		TSharedRef<SWidget> CreateActorsSection(const FArguments& InArgs);
+		TSharedRef<SWidget> CreatePropertiesSection(const FArguments& InArgs);
+
+		/** Sets RootObjectRowData to all non-root nodes from ObjectRowData. */
+		void BuildRootObjectRowData();
+		
+		void GetObjectRowChildren(TSharedPtr<FReplicatedObjectData> ReplicatedObjectData, TFunctionRef<void(TSharedPtr<FReplicatedObjectData>)> ProcessChild);
+		
+		/** Handles how much space the 'Clients' area uses with respect to its expansion state. */
+		SSplitter::ESizeRule GetPropertyAreaSizeRule() const { return bIsPropertyAreaExpanded ? SSplitter::ESizeRule::FractionOfParent : SSplitter::ESizeRule::SizeToContent; }
+		void OnPropertyAreaExpansionChanged(bool bExpanded) { bIsPropertyAreaExpanded = bExpanded; }
+	};
+}
diff --git a/Engine/Plugins/Developer/Concert/ConcertUI/ConcertClientSharedSlate/Source/ConcertClientSharedSlate/Private/Replication/Editor/View/ObjectViewer/SSubobjectAndPropertySection.cpp b/Engine/Plugins/Developer/Concert/ConcertUI/ConcertClientSharedSlate/Source/ConcertClientSharedSlate/Private/Replication/Editor/View/ObjectViewer/SSubobjectAndPropertySection.cpp
--- a/Engine/Plugins/Developer/Concert/ConcertUI/ConcertClientSharedSlate/Source/ConcertClientSharedSlate/Private/Replication/Editor/View/ObjectViewer/SSubobjectAndPropertySection.cpp
+++ b/Engine/Plugins/Developer/Concert/ConcertUI/ConcertClientSharedSlate/Source/ConcertClientSharedSlate/Private/Replication/Editor/View/ObjectViewer/SSubobjectAndPropertySection.cpp
@@ -249,7 +249,7 @@ namespace UE::ConcertClientSharedSlate
 		{
 			ToSort.Sort([this](const TSharedPtr<FReplicatedPropertyData>& Left, const TSharedPtr<FReplicatedPropertyData>& Right)
 			{
-				return SortPropertyRowPredicate.Execute(Left, Right);
+				return SortPropertyRowPredicate.Execute(*Left.Get(), *Right.Get());
 			});
 		}
 	}
diff --git a/Engine/Plugins/Developer/Concert/ConcertUI/ConcertClientSharedSlate/Source/ConcertClientSharedSlate/Private/Replication/Editor/View/ObjectViewer/SSubobjectAndPropertySection.h b/Engine/Plugins/Developer/Concert/ConcertUI/ConcertClientSharedSlate/Source/ConcertClientSharedSlate/Private/Replication/Editor/View/ObjectViewer/SSubobjectAndPropertySection.h
--- a/Engine/Plugins/Developer/Concert/ConcertUI/ConcertClientSharedSlate/Source/ConcertClientSharedSlate/Private/Replication/Editor/View/ObjectViewer/SSubobjectAndPropertySection.h
+++ b/Engine/Plugins/Developer/Concert/ConcertUI/ConcertClientSharedSlate/Source/ConcertClientSharedSlate/Private/Replication/Editor/View/ObjectViewer/SSubobjectAndPropertySection.h
@@ -3,6 +3,7 @@
 #pragma once
 
 #include "Replication/Editor/View/ObjectViewer/Tree/SelectionViewerColumns.h"
+#include "Replication/ReplicationWidgetDelegates.h"
 
 #include "Misc/Optional.h"
 #include "Replication/Data/ConcertPropertySelection.h"
@@ -16,8 +17,6 @@ namespace UE::ConcertClientSharedSlate
 	class SReplicatedPropertiesView;
 	class IReplicationSubobjectView;
 
-	DECLARE_DELEGATE_RetVal_TwoParams(bool, FSortPropertyPredicate, const TSharedPtr<FReplicatedPropertyData>& Left, const TSharedPtr<FReplicatedPropertyData>& Right);
-	
 	/**
 	 * Displays an optional IReplicationSubobjectView and followed by SReplicatedPropertiesView tree view (which is always created).
 	 * 
@@ -87,7 +86,7 @@ namespace UE::ConcertClientSharedSlate
 		/**
 		 * These instances can be subclasses of FReplicatedPropertyData, e.g. FReplicatedPropertyData_Editor.
 		 * Their type can be overridden by subclasses.
-		 * They only have the FReplicatedPropertyData type so they can be passed efficiently to SObjectToPropertyViewer.
+		 * They only have the FReplicatedPropertyData type so they can be passed efficiently to SObjectToPropertyView.
 		 * @see GetPropertyData
 		 */
 		TArray<TSharedPtr<FReplicatedPropertyData>> PropertyRowData;
diff --git a/Engine/Plugins/Developer/Concert/ConcertUI/ConcertClientSharedSlate/Source/ConcertClientSharedSlate/Private/Replication/Editor/View/ObjectViewer/Tree/ReplicationColumn.h b/Engine/Plugins/Developer/Concert/ConcertUI/ConcertClientSharedSlate/Source/ConcertClientSharedSlate/Private/Replication/Editor/View/ObjectViewer/Tree/ReplicationColumn.h
deleted file mode 100644
diff --git a/Engine/Plugins/Developer/Concert/ConcertUI/ConcertClientSharedSlate/Source/ConcertClientSharedSlate/Private/Replication/Editor/View/ObjectViewer/Tree/SReplicationColumnRow.h b/Engine/Plugins/Developer/Concert/ConcertUI/ConcertClientSharedSlate/Source/ConcertClientSharedSlate/Private/Replication/Editor/View/ObjectViewer/Tree/SReplicationColumnRow.h
--- a/Engine/Plugins/Developer/Concert/ConcertUI/ConcertClientSharedSlate/Source/ConcertClientSharedSlate/Private/Replication/Editor/View/ObjectViewer/Tree/SReplicationColumnRow.h
+++ b/Engine/Plugins/Developer/Concert/ConcertUI/ConcertClientSharedSlate/Source/ConcertClientSharedSlate/Private/Replication/Editor/View/ObjectViewer/Tree/SReplicationColumnRow.h
@@ -2,7 +2,7 @@
 
 #pragma once
 
-#include "ReplicationColumn.h"
+#include "Replication/Editor/View/ReplicationColumn.h"
 #include "Widgets/Layout/SBox.h"
 #include "Widgets/SNullWidget.h"
 #include "Widgets/Views/SExpanderArrow.h"
@@ -35,7 +35,7 @@ namespace UE::ConcertClientSharedSlate
 			SLATE_EVENT(FGetColumn, ColumnGetter)
 
 			/** The data to pass to TReplicationColumn::BuildColumnWidget. */
-			SLATE_ARGUMENT(TListItemType, RowData)
+			SLATE_ARGUMENT(TSharedPtr<TListItemType>, RowData)
 		
 			/** The name of the column that will have the SExpanderArrow for the tree view. */
 			SLATE_ARGUMENT(FName, ExpandableColumnLabel)
@@ -66,7 +66,7 @@ namespace UE::ConcertClientSharedSlate
 				return SNullWidget::NullWidget;
 			}
 			
-			const TSharedRef<SWidget> ColumnWidget = Column->BuildColumnWidget({ HighlightText, RowData });
+			const TSharedRef<SWidget> ColumnWidget = Column->BuildColumnWidget({ HighlightText, *RowData.Get() });
 			const bool bNeedsExpanderArrow = ColumnName == ExpandableColumnLabel;
 			if (!bNeedsExpanderArrow)
 			{
@@ -99,7 +99,7 @@ namespace UE::ConcertClientSharedSlate
 		
 		FGetColumn ColumnGetterDelegate;
 		TSharedPtr<FText> HighlightText;
-		TListItemType RowData;
+		TSharedPtr<TListItemType> RowData;
 		FName ExpandableColumnLabel;
 		float RowHeight;
 	};
diff --git a/Engine/Plugins/Developer/Concert/ConcertUI/ConcertClientSharedSlate/Source/ConcertClientSharedSlate/Private/Replication/Editor/View/ObjectViewer/Tree/SReplicationTreeView.h b/Engine/Plugins/Developer/Concert/ConcertUI/ConcertClientSharedSlate/Source/ConcertClientSharedSlate/Private/Replication/Editor/View/ObjectViewer/Tree/SReplicationTreeView.h
--- a/Engine/Plugins/Developer/Concert/ConcertUI/ConcertClientSharedSlate/Source/ConcertClientSharedSlate/Private/Replication/Editor/View/ObjectViewer/Tree/SReplicationTreeView.h
+++ b/Engine/Plugins/Developer/Concert/ConcertUI/ConcertClientSharedSlate/Source/ConcertClientSharedSlate/Private/Replication/Editor/View/ObjectViewer/Tree/SReplicationTreeView.h
@@ -2,7 +2,7 @@
 
 #pragma once
 
-#include "ReplicationColumn.h"
+#include "Replication/Editor/View/ReplicationColumn.h"
 #include "SReplicationColumnRow.h"
 
 #include "Algo/RemoveIf.h"
@@ -26,16 +26,16 @@ namespace UE::ConcertClientSharedSlate
 	{
 	public:
 
-		DECLARE_DELEGATE_OneParam(FDeleteItems, const TArray<TItemType>& SelectedItems);
-		DECLARE_DELEGATE_TwoParams(FGetItemChildren, TItemType Item, TFunctionRef<void(TItemType)> ProcessChild);
+		DECLARE_DELEGATE_OneParam(FDeleteItems, const TArray<TSharedPtr<TItemType>>& SelectedItems);
+		DECLARE_DELEGATE_TwoParams(FGetItemChildren, TSharedPtr<TItemType> Item, TFunctionRef<void(TSharedPtr<TItemType>)> ProcessChild);
 		DECLARE_DELEGATE(FOnSelectionChanged);
-		DECLARE_DELEGATE_RetVal_OneParam(bool, FCustomFilter, const TItemType& Item);
+		DECLARE_DELEGATE_RetVal_OneParam(bool, FCustomFilter, const TSharedPtr<TItemType>& Item);
 
 		SLATE_BEGIN_ARGS(SReplicationTreeView<TItemType>)
 			: _SelectionMode(ESelectionMode::Single)
 		{}
 			/** The items to display */
-			SLATE_ARGUMENT(TArray<TItemType>*, RootItemsSource)
+			SLATE_ARGUMENT(TArray<TSharedPtr<TItemType>>*, RootItemsSource)
 
 			/** Gets an items children for the tree view */
 			SLATE_EVENT(FGetItemChildren, OnGetChildren)
@@ -77,7 +77,7 @@ namespace UE::ConcertClientSharedSlate
 			ExpandableColumnId = InArgs._ExpandableColumnLabel;
 			
 			SearchText = MakeShared<FText>();
-			SearchTextFilter = MakeShared<TTextFilter<const TItemType&>>(TTextFilter<const TItemType&>::FItemToStringArray::CreateSP(this, &SReplicationTreeView::PopulateSearchStrings));
+			SearchTextFilter = MakeShared<TTextFilter<const TSharedPtr<TItemType>&>>(TTextFilter<const TSharedPtr<TItemType>&>::FItemToStringArray::CreateSP(this, &SReplicationTreeView::PopulateSearchStrings));
 			SearchTextFilter->OnChanged().AddSP(this, &SReplicationTreeView::OnFilterChanged);
 			
 			ChildSlot
@@ -138,20 +138,20 @@ namespace UE::ConcertClientSharedSlate
 			OnFilterChanged();
 		}
 		
-		void SetSelectedItems(const TArray<TItemType>& ObjectsToSelect, bool bIsSelected)
+		void SetSelectedItems(const TArray<TSharedPtr<TItemType>>& ObjectsToSelect, bool bIsSelected)
 		{
 			TreeView->ClearSelection();
 			TreeView->SetItemSelection(ObjectsToSelect, bIsSelected);
 		}
-		void SetExpandedItems(const TArray<TItemType>& ObjectsToSelect, bool bIsExpanded)
+		void SetExpandedItems(const TArray<TSharedPtr<TItemType>>& ObjectsToSelect, bool bIsExpanded)
 		{
-			for (TItemType Item : ObjectsToSelect)
+			for (TSharedPtr<TItemType> Item : ObjectsToSelect)
 			{
 				TreeView->SetItemExpansion(Item, bIsExpanded);
 			}
 		}
 		
-		TArray<TItemType> GetSelectedItems() const { return TreeView->GetSelectedItems(); }
+		TArray<TSharedPtr<TItemType>> GetSelectedItems() const { return TreeView->GetSelectedItems(); }
 		
 		virtual FReply OnKeyDown(const FGeometry& MyGeometry, const FKeyEvent& InKeyEvent) override;
 
@@ -162,17 +162,17 @@ namespace UE::ConcertClientSharedSlate
 		/** Used to highlight text in text widgets */
 		TSharedPtr<FText> SearchText;
 		/** Performs text search */
-		TSharedPtr<TTextFilter<const TItemType&>> SearchTextFilter;
+		TSharedPtr<TTextFilter<const TSharedPtr<TItemType>&>> SearchTextFilter;
 
 		/** ListView's header row */
 		TSharedPtr<SHeaderRow> HeaderRow;
 		/** Displays the contents */
-		TSharedPtr<STreeView<TItemType>> TreeView;
+		TSharedPtr<STreeView<TSharedPtr<TItemType>>> TreeView;
 		/** The name of the column which will have the SExpandableArrow widget for the tree view. */
 		FName ExpandableColumnId;
 
-		TArray<TItemType>* AllRootItems = nullptr;
-		TArray<TItemType> FilteredRootItems;
+		TArray<TSharedPtr<TItemType>>* AllRootItems = nullptr;
+		TArray<TSharedPtr<TItemType>> FilteredRootItems;
 
 		/** Callback for getting an item's children. */
 		FGetItemChildren OnGetChildrenDelegate;
@@ -184,15 +184,15 @@ namespace UE::ConcertClientSharedSlate
 
 		TSharedRef<SWidget> CreateTreeView(const FArguments& InArgs);
 		TSharedRef<SHeaderRow> CreateHeaderRow(const FArguments& InArgs);
-		TSharedRef<ITableRow> OnGenerateRowWidget(TItemType Item, const TSharedRef<STableViewBase>& OwnerTable);
-		void GetRowChildren(TItemType Item, TArray<TItemType>& OutChildren);
+		TSharedRef<ITableRow> OnGenerateRowWidget(TSharedPtr<TItemType> Item, const TSharedRef<STableViewBase>& OwnerTable);
+		void GetRowChildren(TSharedPtr<TItemType> Item, TArray<TSharedPtr<TItemType>>& OutChildren);
 		
 		void OnSearchTextCommitted(const FText& InFilterText, ETextCommit::Type CommitType);
 		void OnSearchTextChanged(const FText& InSearchText);
 
-		void PopulateSearchStrings(const TItemType& Item, TArray<FString>& OutSearchStrings);
+		void PopulateSearchStrings(const TSharedPtr<TItemType>& Item, TArray<FString>& OutSearchStrings);
 		void OnFilterChanged();
-		bool PassesFilters(const TItemType& Item);
+		bool PassesFilters(const TSharedPtr<TItemType>& Item);
 	};
 
 	template <typename TItemType>
@@ -215,7 +215,7 @@ namespace UE::ConcertClientSharedSlate
 			.BorderBackgroundColor(FSlateColor(FLinearColor(0.6, 0.6, 0.6)))
 			.Padding(0)
 			[
-				SAssignNew(TreeView, STreeView<TItemType>)
+				SAssignNew(TreeView, STreeView<TSharedPtr<TItemType>>)
 				.OnGetChildren(this, &SReplicationTreeView::GetRowChildren)
 				.TreeItemsSource(&FilteredRootItems)
 				.OnGenerateRow(this, &SReplicationTreeView::OnGenerateRowWidget)
@@ -249,7 +249,7 @@ namespace UE::ConcertClientSharedSlate
 	}
 
 	template <typename TItemType>
-	TSharedRef<ITableRow> SReplicationTreeView<TItemType>::OnGenerateRowWidget(TItemType Item, const TSharedRef<STableViewBase>& OwnerTable)
+	TSharedRef<ITableRow> SReplicationTreeView<TItemType>::OnGenerateRowWidget(TSharedPtr<TItemType> Item, const TSharedRef<STableViewBase>& OwnerTable)
 	{
 		const typename SReplicationColumnRow<TItemType>::FGetColumn ColumnGetter = SReplicationColumnRow<TItemType>::FGetColumn::CreateLambda([this](const FName& ColumnId) -> const TReplicationColumn<TItemType>*
 		{
@@ -273,11 +273,11 @@ namespace UE::ConcertClientSharedSlate
 	}
 
 	template <typename TItemType>
-	void SReplicationTreeView<TItemType>::GetRowChildren(TItemType Item, TArray<TItemType>& OutChildren)
+	void SReplicationTreeView<TItemType>::GetRowChildren(TSharedPtr<TItemType> Item, TArray<TSharedPtr<TItemType>>& OutChildren)
 	{
 		if (OnGetChildrenDelegate.IsBound())
 		{
-			OnGetChildrenDelegate.Execute(Item, [this, &OutChildren](TItemType ItemToAdd)
+			OnGetChildrenDelegate.Execute(Item, [this, &OutChildren](TSharedPtr<TItemType> ItemToAdd)
 			{
 				if (PassesFilters(ItemToAdd))
 				{
@@ -305,12 +305,12 @@ namespace UE::ConcertClientSharedSlate
 	}
 
 	template <typename TItemType>
-	void SReplicationTreeView<TItemType>::PopulateSearchStrings(const TItemType& Item, TArray<FString>& OutSearchStrings)
+	void SReplicationTreeView<TItemType>::PopulateSearchStrings(const TSharedPtr<TItemType>& Item, TArray<FString>& OutSearchStrings)
 	{
 		for (const SHeaderRow::FColumn& Column : HeaderRow->GetColumns())
 		{
 			const TReplicationColumn<TItemType>& CastColumn = static_cast<const TReplicationColumn<TItemType>&>(Column);
-			CastColumn.ExecutePopulateSearchString(Item, OutSearchStrings);
+			CastColumn.ExecutePopulateSearchString(*Item.Get(), OutSearchStrings);
 		}
 	}
 
@@ -318,13 +318,13 @@ namespace UE::ConcertClientSharedSlate
 	void SReplicationTreeView<TItemType>::OnFilterChanged()
 	{
 		// Try preserving the selected activity.
-		TArray<TItemType> SelectedItems = TreeView->GetSelectedItems();
+		TArray<TSharedPtr<TItemType>> SelectedItems = TreeView->GetSelectedItems();
 
 		// Reset the list of displayed activities.
 		FilteredRootItems.Reset(AllRootItems->Num());
 
 		// Apply the filter.
-		for (const TItemType& Activity : *AllRootItems)
+		for (const TSharedPtr<TItemType>& Activity : *AllRootItems)
 		{
 			if (PassesFilters(Activity))
 			{
@@ -333,7 +333,7 @@ namespace UE::ConcertClientSharedSlate
 		}
 
 		// Restore/reset the selected activity.
-		SelectedItems.SetNum(Algo::RemoveIf(SelectedItems, [this](const TItemType& Item){ return !FilteredRootItems.Contains(Item); }));
+		SelectedItems.SetNum(Algo::RemoveIf(SelectedItems, [this](const TSharedPtr<TItemType>& Item){ return !FilteredRootItems.Contains(Item); }));
 		if (!SelectedItems.IsEmpty())
 		{
 			TreeView->SetItemSelection(SelectedItems, true); // Restore previous selection.
@@ -344,7 +344,7 @@ namespace UE::ConcertClientSharedSlate
 	}
 
 	template <typename TItemType>
-	bool SReplicationTreeView<TItemType>::PassesFilters(const TItemType& Item)
+	bool SReplicationTreeView<TItemType>::PassesFilters(const TSharedPtr<TItemType>& Item)
 	{
 		return SearchTextFilter->PassesFilter(Item)
 			&& (!CustomFilterDelegate.IsBound() || CustomFilterDelegate.Execute(Item));
diff --git a/Engine/Plugins/Developer/Concert/ConcertUI/ConcertClientSharedSlate/Source/ConcertClientSharedSlate/Private/Replication/Editor/View/ObjectViewer/Tree/SelectionViewerColumns.cpp b/Engine/Plugins/Developer/Concert/ConcertUI/ConcertClientSharedSlate/Source/ConcertClientSharedSlate/Private/Replication/Editor/View/ObjectViewer/Tree/SelectionViewerColumns.cpp
--- a/Engine/Plugins/Developer/Concert/ConcertUI/ConcertClientSharedSlate/Source/ConcertClientSharedSlate/Private/Replication/Editor/View/ObjectViewer/Tree/SelectionViewerColumns.cpp
+++ b/Engine/Plugins/Developer/Concert/ConcertUI/ConcertClientSharedSlate/Source/ConcertClientSharedSlate/Private/Replication/Editor/View/ObjectViewer/Tree/SelectionViewerColumns.cpp
@@ -6,10 +6,11 @@
 #include "Replication/Editor/Model/IEditableObjectToPropertiesModel.h"
 #include "Replication/Editor/Model/ReplicatedPropertyData.h"
 #include "Replication/Editor/Model/ReplicatedObjectData.h"
-#include "Replication/Editor/View/ObjectEditor/SObjectToPropertyEditor.h"
+#include "Replication/Editor/View/ObjectEditor/SBaseReplicationStreamEditor.h"
 #include "Replication/PropertyChainUtils.h"
 
 #include "Internationalization/Internationalization.h"
+#include "GameFramework/Actor.h"
 #include "Textures/SlateIcon.h"
 #include "Widgets/Images/SImage.h"
 #include "Widgets/SBoxPanel.h"
@@ -36,7 +37,7 @@ namespace UE::ConcertClientSharedSlate::ReplicationObjectColumns
 						.VAlign(VAlign_Center)
 						[
 							SNew(SImage)
-							.Image(DisplayUtils::GetObjectIcon(*Model, Args.RowData->GetObjectPath()).GetOptionalIcon())
+							.Image(DisplayUtils::GetObjectIcon(*Model, Args.RowData.GetObjectPath()).GetOptionalIcon())
 						];
 				})
 				.ColumnSortOrder(static_cast<int32>(EReplicationColumnOrder::Icon)),
@@ -54,11 +55,11 @@ namespace UE::ConcertClientSharedSlate::ReplicationObjectColumns
 				{
 					return SNew(STextBlock)
 						.HighlightText(TAttribute<FText>::CreateLambda([HighlightText = Args.HighlightText](){ return *HighlightText; }))
-						.Text(DisplayUtils::GetObjectDisplayText(Args.RowData->GetObjectPath()));
+						.Text(DisplayUtils::GetObjectDisplayText(Args.RowData.GetObjectPath()));
 				})
-				.PopulateSearchItems_Lambda([](const TSharedPtr<FReplicatedObjectData>& ObjectData, TArray<FString>& InOutSearchStrings)
+				.PopulateSearchItems_Lambda([](const FReplicatedObjectData& ObjectData, TArray<FString>& InOutSearchStrings)
 				{
-					InOutSearchStrings.Add(DisplayUtils::GetObjectDisplayText(ObjectData->GetObjectPath()).ToString());
+					InOutSearchStrings.Add(DisplayUtils::GetObjectDisplayText(ObjectData.GetObjectPath()).ToString());
 				})
 				.ColumnSortOrder(static_cast<int32>(EReplicationColumnOrder::Label)),
 			SHeaderRow::Column(LabelColumnId)
@@ -75,11 +76,11 @@ namespace UE::ConcertClientSharedSlate::ReplicationObjectColumns
 				{
 					return SNew(STextBlock)
 						.HighlightText(TAttribute<FText>::CreateLambda([HighlightText = Args.HighlightText](){ return *HighlightText; }))
-						.Text(DisplayUtils::GetObjectTypeText(*Model, Args.RowData->GetObjectPath()));
+						.Text(DisplayUtils::GetObjectTypeText(*Model, Args.RowData.GetObjectPath()));
 				})
-				.PopulateSearchItems_Lambda([Model](const TSharedPtr<FReplicatedObjectData>& ObjectData, TArray<FString>& InOutSearchStrings)
+				.PopulateSearchItems_Lambda([Model](const FReplicatedObjectData& ObjectData, TArray<FString>& InOutSearchStrings)
 				{
-					InOutSearchStrings.Add(DisplayUtils::GetObjectTypeText(Model.Get(), ObjectData->GetObjectPath()).ToString());
+					InOutSearchStrings.Add(DisplayUtils::GetObjectTypeText(Model.Get(), ObjectData.GetObjectPath()).ToString());
 				})
 				.ColumnSortOrder(static_cast<int32>(EReplicationColumnOrder::Type)),
 			SHeaderRow::Column(TypeColumnId)
@@ -94,38 +95,8 @@ namespace UE::ConcertClientSharedSlate::ReplicationObjectColumns
 
 namespace UE::ConcertClientSharedSlate::ReplicationPropertyColumns
 {
-	const FName ReplicatesColumnId = TEXT("ReplicatesColumn");
 	const FName LabelColumnId = TEXT("LabelColumn");
 	const FName TypeColumnId = TEXT("TypeColumn");
-
-	FReplicationPropertyColumn ReplicatesColumns(
-		FGetPropertyCheckboxState GetPropertyCheckboxStateDelegate,
-		FOnPropertyCheckboxChanged OnPropertyBoxToggledDelegate,
-		const float ColumnWidth
-		)
-	{
-		return FReplicationPropertyColumn(
-			FReplicationPropertyColumn::FArguments()
-				.GenerateWidgetColumn_Lambda([GetPropertyCheckboxStateDelegate, OnPropertyBoxToggledDelegate](const FReplicationPropertyColumn::FBuildArgs& Args)
-				{
-					return SNew(SCheckBox)
-						.ToolTipText(LOCTEXT("ReplicatesCheckbox.Tooltip", "Should replicate?"))
-						.IsChecked_Lambda([GetPropertyCheckboxStateDelegate, RowData = Args.RowData]()
-						{
-							return GetPropertyCheckboxStateDelegate.Execute(RowData->GetProperty());
-						})
-						.OnCheckStateChanged_Lambda([OnPropertyBoxToggledDelegate, RowData = Args.RowData](ECheckBoxState NewState)
-						{
-							const bool bIsChecked = NewState == ECheckBoxState::Checked;
-							OnPropertyBoxToggledDelegate.Execute(bIsChecked, RowData->GetProperty());
-						});
-				})
-				.ColumnSortOrder(static_cast<int32>(EReplicationPropertyColumnOrder::ReplicatesCheckbox)),
-			SHeaderRow::Column(ReplicatesColumnId)
-				.DefaultLabel(FText::GetEmpty())
-				.FixedWidth(ColumnWidth)
-			);
-	}
 	
 	FReplicationPropertyColumn LabelColumn()
 	{
@@ -135,11 +106,11 @@ namespace UE::ConcertClientSharedSlate::ReplicationPropertyColumns
 				{
 					return SNew(STextBlock)
 						.HighlightText(TAttribute<FText>::CreateLambda([HighlightText = Args.HighlightText](){ return *HighlightText; }))
-						.Text(DisplayUtils::GetPropertyDisplayText(Args.RowData->GetProperty()));
+						.Text(DisplayUtils::GetPropertyDisplayText(Args.RowData.GetProperty()));
 				})
-				.PopulateSearchItems_Lambda([](const TSharedPtr<FReplicatedPropertyData>& ObjectData, TArray<FString>& InOutSearchStrings)
+				.PopulateSearchItems_Lambda([](const FReplicatedPropertyData& ObjectData, TArray<FString>& InOutSearchStrings)
 				{
-					InOutSearchStrings.Add(DisplayUtils::GetPropertyDisplayText(ObjectData->GetProperty()).ToString());
+					InOutSearchStrings.Add(DisplayUtils::GetPropertyDisplayText(ObjectData.GetProperty()).ToString());
 				})
 				.ColumnSortOrder(static_cast<int32>(EReplicationPropertyColumnOrder::Label)),
 			SHeaderRow::Column(LabelColumnId)
@@ -150,11 +121,10 @@ namespace UE::ConcertClientSharedSlate::ReplicationPropertyColumns
 	
 	FReplicationPropertyColumn TypeColumn()
 	{
-		static auto GetDisplayText = [](const TSharedPtr<FReplicatedPropertyData>& Args)
+		static auto GetDisplayText = [](const FReplicatedPropertyData& Args)
 		{
-			check(Args);
-			UClass* Class = Args->GetOwningClass().TryLoadClass<UObject>();
-			const FProperty* Property = Class ? ConcertSyncCore::PropertyChain::ResolveProperty(*Class, Args->GetProperty()) : nullptr;
+			UClass* Class = Args.GetOwningClass().TryLoadClass<UObject>();
+			const FProperty* Property = Class ? ConcertSyncCore::PropertyChain::ResolveProperty(*Class, Args.GetProperty()) : nullptr;
 			return Property ? FText::FromString(Property->GetCPPType()) : LOCTEXT("Unknown", "Unknown");	
 		};
 		
@@ -166,7 +136,7 @@ namespace UE::ConcertClientSharedSlate::ReplicationPropertyColumns
 						.HighlightText(TAttribute<FText>::CreateLambda([HighlightText = Args.HighlightText](){ return *HighlightText; }))
 						.Text(GetDisplayText(Args.RowData));
 				})
-				.PopulateSearchItems_Lambda([](const TSharedPtr<FReplicatedPropertyData>& ObjectData, TArray<FString>& InOutSearchStrings)
+				.PopulateSearchItems_Lambda([](const FReplicatedPropertyData& ObjectData, TArray<FString>& InOutSearchStrings)
 				{
 					InOutSearchStrings.Add(GetDisplayText(ObjectData).ToString());
 				})
@@ -177,6 +147,99 @@ namespace UE::ConcertClientSharedSlate::ReplicationPropertyColumns
 			);
 	}
 
+	namespace Private
+	{
+		static ECheckBoxState OnGetPropertyCheckboxState(
+			const FConcertPropertyChain& PropertyChain,
+			const IReplicationStreamViewer& Viewer,
+			const IObjectToPropertiesModel& Model
+			)
+		{
+			const TArray<FSoftObjectPath> SelectedObjectPaths = Viewer.GetObjectsBeingPropertyEdited();
+			return GetPropertyCheckboxStateBasedOnSelection(PropertyChain, SelectedObjectPaths, Model);
+		}
+
+		static void OnPropertyCheckboxChanged(
+			bool bIsChecked,
+			const FConcertPropertyChain& PropertyChain,
+			const IReplicationStreamViewer& Viewer,
+			IEditableObjectToPropertiesModel& Model
+			)
+		{
+			const TArray Properties{ PropertyChain };
+			const TArray<FSoftObjectPath> SelectedObjects = Viewer.GetObjectsBeingPropertyEdited();
+
+			if (bIsChecked)
+			{
+				// We cannot proceed if any of the selected objects cannot be loaded because we cannot obtain its class
+				for (const FSoftObjectPath& Path : SelectedObjects)
+				{
+					UObject* Object = Path.ResolveObject();
+					if (!Object && !Model.ContainsObjects({ Path }))
+					{
+						return;
+					}
+				}
+				
+				for (const FSoftObjectPath& Path : Viewer.GetObjectsBeingPropertyEdited())
+				{
+					UObject* Object = Path.ResolveObject();
+					if (!Model.ContainsObjects({ Path }))
+					{
+						Model.AddObjects({ Object });
+					}
+					Model.AddProperties(Path, Properties);
+				}
+			}
+			else
+			{
+				for (const FSoftObjectPath& Path : Viewer.GetObjectsBeingPropertyEdited())
+				{
+					Model.RemoveProperties(Path, Properties);
+					// If this does not resolve it is not too bad if the subobject is not removed...
+					UObject* Object = Path.ResolveObject();
+					const bool bNeedsToRemoveNonRoot = Object && !Object->IsA<AActor>() && Model.GetNumProperties(Path) == 0;
+					if (bNeedsToRemoveNonRoot)
+					{
+						Model.RemoveObjects({ Path });
+					}
+				}
+			}
+		}
+	}
+	
+	FReplicationPropertyColumn ReplicatesColumns(
+		TWeakPtr<IReplicationStreamViewer> Viewer,
+		TWeakPtr<IEditableObjectToPropertiesModel> Model,
+		const float ColumnWidth,
+		const int32 Priority
+		)
+	{
+		return MakePropertyCheckboxColumn(
+			FGetColumnCheckboxState::CreateLambda(
+				[Viewer, Model](const FConcertPropertyChain& Property)
+				{
+					const TSharedPtr<IReplicationStreamViewer> ViewerPin = Viewer.Pin();
+					const TSharedPtr<IEditableObjectToPropertiesModel> ModelPin = Model.Pin();
+					return ensure(ViewerPin && ModelPin) ? Private::OnGetPropertyCheckboxState(Property, *ViewerPin, *ModelPin) : ECheckBoxState::Undetermined;
+				}),
+			FOnColumnCheckboxChanged::CreateLambda(
+				[Viewer, Model](bool bIsChecked, const FConcertPropertyChain& Property)
+				{
+					const TSharedPtr<IReplicationStreamViewer> ViewerPin = Viewer.Pin();
+					const TSharedPtr<IEditableObjectToPropertiesModel> ModelPin = Model.Pin();
+					if (ensure(ViewerPin && ModelPin))
+					{
+						Private::OnPropertyCheckboxChanged(bIsChecked, Property, *ViewerPin, *ModelPin);
+					}
+				}),
+			FText::GetEmpty(),
+			LOCTEXT("Replicates.ToolTip", "Select whether this property should be replicated"),
+			ColumnWidth, 
+			Priority
+		);
+	}
+
 	ECheckBoxState GetPropertyCheckboxStateBasedOnSelection(const FConcertPropertyChain& Property, TConstArrayView<FSoftObjectPath> Selection, const IObjectToPropertiesModel& Model)
 	{
 		ECheckBoxState CheckBoxState = ECheckBoxState::Undetermined;
@@ -199,6 +262,27 @@ namespace UE::ConcertClientSharedSlate::ReplicationPropertyColumns
 		}
 		return CheckBoxState;
 	}
+	
+	bool SortBySelectionThenByName_PropertyPredicate(
+		const TArray<FSoftObjectPath>& SelectedObjects,
+		const IObjectToPropertiesModel& Model,
+		const FReplicatedPropertyData& Left,
+		const FReplicatedPropertyData& Right
+		)
+	{
+		const ECheckBoxState LeftCheckboxState = GetPropertyCheckboxStateBasedOnSelection(Left.GetProperty(), SelectedObjects, Model);
+		const ECheckBoxState RightCheckboxState = GetPropertyCheckboxStateBasedOnSelection(Right.GetProperty(), SelectedObjects, Model);
+		
+		// Secondary sort by name
+		if (LeftCheckboxState == RightCheckboxState)
+		{
+			return DisplayUtils::GetPropertyDisplayString(Left.GetProperty()) < DisplayUtils::GetPropertyDisplayString(Right.GetProperty());
+		}
+
+		// Selected properties should appear first
+		return LeftCheckboxState == ECheckBoxState::Checked
+			&& (RightCheckboxState == ECheckBoxState::Unchecked || RightCheckboxState == ECheckBoxState::Undetermined);
+	}
 }
 
 #undef LOCTEXT_NAMESPACE
\ No newline at end of file
diff --git a/Engine/Plugins/Developer/Concert/ConcertUI/ConcertClientSharedSlate/Source/ConcertClientSharedSlate/Private/Replication/Editor/View/ObjectViewer/Tree/SelectionViewerColumns.h b/Engine/Plugins/Developer/Concert/ConcertUI/ConcertClientSharedSlate/Source/ConcertClientSharedSlate/Private/Replication/Editor/View/ObjectViewer/Tree/SelectionViewerColumns.h
--- a/Engine/Plugins/Developer/Concert/ConcertUI/ConcertClientSharedSlate/Source/ConcertClientSharedSlate/Private/Replication/Editor/View/ObjectViewer/Tree/SelectionViewerColumns.h
+++ b/Engine/Plugins/Developer/Concert/ConcertUI/ConcertClientSharedSlate/Source/ConcertClientSharedSlate/Private/Replication/Editor/View/ObjectViewer/Tree/SelectionViewerColumns.h
@@ -2,35 +2,16 @@
 
 #pragma once
 
-#include "Containers/ArrayView.h"
-#include "ReplicationColumn.h"
+#include "Replication/Editor/View/PredefinedReplicationColumns.h"
 #include "Templates/SharedPointer.h"
 
-struct FConcertPropertyChain;
-
 namespace UE::ConcertClientSharedSlate
 {
-	class FReplicatedObjectData;
-	class FReplicatedPropertyData;
-	class IEditableObjectToPropertiesModel;
 	class IObjectToPropertiesModel;
-	class SObjectToPropertyEditor;
 }
 
 namespace UE::ConcertClientSharedSlate::ReplicationObjectColumns
 {
-	using FReplicationObjectColumn = TReplicationColumn<TSharedPtr<FReplicatedObjectData>>;
-	
-	enum class EReplicationColumnOrder : int32
-	{
-		/** Displays the class icon */
-		Icon = 10,
-		/** Label of the object */
-		Label = 20,
-		/** Class of the object */
-		Type = 30,
-	};
-	
 	extern const FName IconColumnId;
 	extern const FName LabelColumnId;
 	extern const FName TypeColumnId;
@@ -42,33 +23,22 @@ namespace UE::ConcertClientSharedSlate::ReplicationObjectColumns
 
 namespace UE::ConcertClientSharedSlate::ReplicationPropertyColumns
 {
-	using FReplicationPropertyColumn = TReplicationColumn<TSharedPtr<FReplicatedPropertyData>>;
-	
-	enum class EReplicationPropertyColumnOrder : int32
-	{
-		/** The checkbox in SObjectToPropertyEditor determining whether the property is in the selection*/
-		ReplicatesCheckbox = 0,
-		/** Label of the property */
-		Label = 10,
-		/** Type of the property */
-		Type = 20
-	};
-	
-	/** The checkbox in SObjectToPropertyEditor determining whether the property is in the selection*/
 	extern const FName ReplicatesColumnId;
 	extern const FName LabelColumnId;
 	extern const FName TypeColumnId;
-
-	DECLARE_DELEGATE_RetVal_OneParam(ECheckBoxState, FGetPropertyCheckboxState, const FConcertPropertyChain& /*Property*/);
-	DECLARE_DELEGATE_TwoParams(FOnPropertyCheckboxChanged, bool /*bIsChecked*/, const FConcertPropertyChain& /*Property*/);
-
-	FReplicationPropertyColumn ReplicatesColumns(
-		FGetPropertyCheckboxState GetPropertyCheckboxStateDelegate,
-		FOnPropertyCheckboxChanged OnPropertyBoxToggledDelegate,
-		const float ColumnWidth = 20.f
-		);
+	
 	FReplicationPropertyColumn LabelColumn();
 	FReplicationPropertyColumn TypeColumn();
+	/**
+	 * A checkbox that is placed at the beginning of every property.
+	 * Checking & unchecking adds & removes the property to the selected objects' property mapping, respectively.
+	 */
+	CONCERTCLIENTSHAREDSLATE_API FReplicationPropertyColumn ReplicatesColumns(
+		TWeakPtr<IReplicationStreamViewer> Viewer,
+		TWeakPtr<IEditableObjectToPropertiesModel> Model,
+		const float ColumnWidth = 20.f,
+		const int32 Priority = static_cast<int32>(EReplicationPropertyColumnOrder::ReplicatesCheckbox)
+		);
 
 	/**
 	 * Goes through all selected objects, checks whether the property is checked on it or not, and returns a checkbox state.
@@ -76,5 +46,16 @@ namespace UE::ConcertClientSharedSlate::ReplicationPropertyColumns
 	 *
 	 * This is used by ReplicatesColumns.
 	 */
-	ECheckBoxState GetPropertyCheckboxStateBasedOnSelection(const FConcertPropertyChain& Property, TConstArrayView<FSoftObjectPath> Selection, const IObjectToPropertiesModel& Model);
+	CONCERTCLIENTSHAREDSLATE_API ECheckBoxState GetPropertyCheckboxStateBasedOnSelection(
+		const FConcertPropertyChain& Property,
+		TConstArrayView<FSoftObjectPath> Selection,
+		const IObjectToPropertiesModel& Model
+		);
+	/** Util for sorting property data by whether its ReplicatesColumns() column is checked. */
+	CONCERTCLIENTSHAREDSLATE_API bool SortBySelectionThenByName_PropertyPredicate(
+		const TArray<FSoftObjectPath>& SelectedObjects,
+		const IObjectToPropertiesModel& Model,
+		const FReplicatedPropertyData& Left,
+		const FReplicatedPropertyData& Right
+		);
 }
diff --git a/Engine/Plugins/Developer/Concert/ConcertUI/ConcertClientSharedSlate/Source/ConcertClientSharedSlate/Private/Replication/Editor/View/PredefinedReplicationColumns.cpp b/Engine/Plugins/Developer/Concert/ConcertUI/ConcertClientSharedSlate/Source/ConcertClientSharedSlate/Private/Replication/Editor/View/PredefinedReplicationColumns.cpp
new file mode 100644
--- /dev/null
+++ b/Engine/Plugins/Developer/Concert/ConcertUI/ConcertClientSharedSlate/Source/ConcertClientSharedSlate/Private/Replication/Editor/View/PredefinedReplicationColumns.cpp
@@ -0,0 +1,53 @@
+// Copyright Epic Games, Inc. All Rights Reserved.
+
+#include "Replication/Editor/View/PredefinedReplicationColumns.h"
+
+#include "Replication/Editor/Model/ReplicatedPropertyData.h"
+
+#include "Internationalization/Internationalization.h"
+#include "Textures/SlateIcon.h"
+#include "Widgets/Images/SImage.h"
+#include "Widgets/SBoxPanel.h"
+#include "Widgets/Text/STextBlock.h"
+#include "Widgets/Input/SCheckBox.h"
+
+#define LOCTEXT_NAMESPACE "ReplicationPropertyColumns"
+
+namespace UE::ConcertClientSharedSlate::ReplicationPropertyColumns
+{
+	const FName ReplicatesColumnId = TEXT("ReplicatesColumn");
+
+	FReplicationPropertyColumn MakePropertyCheckboxColumn(
+		FGetColumnCheckboxState GetPropertyCheckboxStateDelegate,
+		FOnColumnCheckboxChanged OnPropertyBoxToggledDelegate,
+		FText DefaultLabel,
+		FText ToolTipText,
+		const float ColumnWidth,
+		const int32 Priority
+		)
+	{
+		return FReplicationPropertyColumn(
+			FReplicationPropertyColumn::FArguments()
+				.GenerateWidgetColumn_Lambda([GetPropertyCheckboxStateDelegate, OnPropertyBoxToggledDelegate, ToolTipText = MoveTemp(ToolTipText)](const FReplicationPropertyColumn::FBuildArgs& Args)
+				{
+					return SNew(SCheckBox)
+						.ToolTipText(ToolTipText)
+						.IsChecked_Lambda([GetPropertyCheckboxStateDelegate, RowData = Args.RowData]()
+						{
+							return GetPropertyCheckboxStateDelegate.Execute(RowData.GetProperty());
+						})
+						.OnCheckStateChanged_Lambda([OnPropertyBoxToggledDelegate, RowData = Args.RowData](ECheckBoxState NewState)
+						{
+							const bool bIsChecked = NewState == ECheckBoxState::Checked;
+							OnPropertyBoxToggledDelegate.Execute(bIsChecked, RowData.GetProperty());
+						});
+				})
+				.ColumnSortOrder(Priority),
+			SHeaderRow::Column(ReplicatesColumnId)
+				.DefaultLabel(DefaultLabel)
+				.FixedWidth(ColumnWidth)
+			);
+	}
+}
+
+#undef LOCTEXT_NAMESPACE
\ No newline at end of file
diff --git a/Engine/Plugins/Developer/Concert/ConcertUI/ConcertClientSharedSlate/Source/ConcertClientSharedSlate/Private/Replication/Editor/View/SReplicationStreamEditor.cpp b/Engine/Plugins/Developer/Concert/ConcertUI/ConcertClientSharedSlate/Source/ConcertClientSharedSlate/Private/Replication/Editor/View/SReplicationStreamEditor.cpp
deleted file mode 100644
diff --git a/Engine/Plugins/Developer/Concert/ConcertUI/ConcertClientSharedSlate/Source/ConcertClientSharedSlate/Private/Replication/Editor/View/SReplicationStreamEditor.h b/Engine/Plugins/Developer/Concert/ConcertUI/ConcertClientSharedSlate/Source/ConcertClientSharedSlate/Private/Replication/Editor/View/SReplicationStreamEditor.h
deleted file mode 100644
diff --git a/Engine/Plugins/Developer/Concert/ConcertUI/ConcertClientSharedSlate/Source/ConcertClientSharedSlate/Private/Replication/ReplicationWidgetFactories.cpp b/Engine/Plugins/Developer/Concert/ConcertUI/ConcertClientSharedSlate/Source/ConcertClientSharedSlate/Private/Replication/ReplicationWidgetFactories.cpp
--- a/Engine/Plugins/Developer/Concert/ConcertUI/ConcertClientSharedSlate/Source/ConcertClientSharedSlate/Private/Replication/ReplicationWidgetFactories.cpp
+++ b/Engine/Plugins/Developer/Concert/ConcertUI/ConcertClientSharedSlate/Source/ConcertClientSharedSlate/Private/Replication/ReplicationWidgetFactories.cpp
@@ -3,7 +3,8 @@
 #include "Replication/ReplicationWidgetFactories.h"
 
 #include "Editor/Model/TransactionalPropertySelectionModel.h"
-#include "Editor/View/SReplicationStreamEditor.h"
+#include "Editor/View/ObjectEditor/SDefaultReplicationStreamEditor.h"
+#include "Editor/View/ObjectEditor/SBaseReplicationStreamEditor.h"
 #include "Editor/View/SubobjectView/SUnrealEditorSubobjectView.h"
 
 namespace UE::ConcertClientSharedSlate
@@ -13,11 +14,33 @@ namespace UE::ConcertClientSharedSlate
 		return SNew(SUnrealEditorSubobjectView);
 	}
 	
-	TSharedRef<IReplicationEditorView> CreateEditor(FCreateEditorParams Params)
+	TSharedRef<IReplicationStreamEditor> CreateEditor(FCreateEditorParams Params)
 	{
-		return SNew(SReplicationStreamEditor, MoveTemp(Params));
+		return SNew(SBaseReplicationStreamEditor, Params.DataModel, Params.ObjectSource, Params.PropertySource)
+			.AdditionalObjectColumns(Params.AdditionalObjectColumns)
+			.AdditionalPropertyColumns(Params.AdditionalPropertyColumns)
+			.SubobjectView(Params.SubobjectView)
+			.OnExtendObjectsContextMenu(Params.OnExtendObjectsContextMenu)
+			.SortPropertyRowPredicate(Params.SortPropertyRowPredicate)
+			.LeftOfObjectSearchBar() [ Params.LeftOfObjectSearchBar.Widget ]
+			.LeftOfPropertySearchBar() [ Params.LeftOfPropertySearchBar.Widget ];
 	}
-	
+
+	TSharedRef<IReplicationStreamEditor> CreateDefaultStreamEditor(FCreateEditorParams Params)
+	{
+		// Default editor relies on a subobject view
+		Params.SubobjectView = Params.SubobjectView ? Params.SubobjectView : CreateUnrealEditorSubobjectView();
+		
+		return SNew(SDefaultReplicationStreamEditor, Params.DataModel, Params.ObjectSource, Params.PropertySource)
+			.AdditionalObjectColumns(Params.AdditionalObjectColumns)
+			.AdditionalPropertyColumns(Params.AdditionalPropertyColumns)
+			.SubobjectView(Params.SubobjectView)
+			.OnExtendObjectsContextMenu(Params.OnExtendObjectsContextMenu)
+			.SortPropertyRowPredicate(Params.SortPropertyRowPredicate)
+			.LeftOfObjectSearchBar() [ Params.LeftOfObjectSearchBar.Widget ]
+			.LeftOfPropertySearchBar() [ Params.LeftOfPropertySearchBar.Widget ];
+	}
+
 	TSharedRef<IEditableObjectToPropertiesModel> CreatePropertySelectionModel(
 		UObject& OwnerObject,
 		TAttribute<FObjectReplicationMap*> ReplicationMapAttribute,
diff --git a/Engine/Plugins/Developer/Concert/ConcertUI/ConcertClientSharedSlate/Source/ConcertClientSharedSlate/Public/Replication/Editor/Model/ReplicatedObjectData.h b/Engine/Plugins/Developer/Concert/ConcertUI/ConcertClientSharedSlate/Source/ConcertClientSharedSlate/Public/Replication/Editor/Model/ReplicatedObjectData.h
new file mode 100644
--- /dev/null
+++ b/Engine/Plugins/Developer/Concert/ConcertUI/ConcertClientSharedSlate/Source/ConcertClientSharedSlate/Public/Replication/Editor/Model/ReplicatedObjectData.h
@@ -0,0 +1,25 @@
+// Copyright Epic Games, Inc. All Rights Reserved.
+
+#pragma once
+
+#include "UObject/SoftObjectPath.h"
+
+namespace UE::ConcertClientSharedSlate
+{
+	/** Instanced for each object row in SObjectToPropertyView.*/
+	class FReplicatedObjectData
+	{
+	public:
+
+		FReplicatedObjectData(FSoftObjectPath ObjectPath)
+			: ObjectPath(MoveTemp(ObjectPath))
+		{}
+		
+		const FSoftObjectPath& GetObjectPath() const { return ObjectPath; }
+
+	private:
+
+		/** The replicated object */
+		FSoftObjectPath ObjectPath;
+	};
+}
\ No newline at end of file
diff --git a/Engine/Plugins/Developer/Concert/ConcertUI/ConcertClientSharedSlate/Source/ConcertClientSharedSlate/Public/Replication/Editor/Model/ReplicatedPropertyData.h b/Engine/Plugins/Developer/Concert/ConcertUI/ConcertClientSharedSlate/Source/ConcertClientSharedSlate/Public/Replication/Editor/Model/ReplicatedPropertyData.h
new file mode 100644
--- /dev/null
+++ b/Engine/Plugins/Developer/Concert/ConcertUI/ConcertClientSharedSlate/Source/ConcertClientSharedSlate/Public/Replication/Editor/Model/ReplicatedPropertyData.h
@@ -0,0 +1,30 @@
+// Copyright Epic Games, Inc. All Rights Reserved.
+
+#pragma once
+
+#include "Replication/Data/ConcertPropertySelection.h"
+
+namespace UE::ConcertClientSharedSlate
+{
+	/** Instanced for each property row in SObjectToPropertyView.*/
+	class FReplicatedPropertyData
+	{
+	public:
+		
+		FReplicatedPropertyData(FSoftClassPath OwningClass, FConcertPropertyChain Object)
+			: OwningClass(MoveTemp(OwningClass))
+			, Property(MoveTemp(Object))
+		{}
+		
+		const FConcertPropertyChain& GetProperty() const { return Property; }
+		const FSoftClassPath& GetOwningClass() const { return OwningClass; }
+
+	private:
+
+		/** The class with which the FProperty can be determined. */
+		FSoftClassPath OwningClass;
+		
+		/** The property to be replicated */
+		FConcertPropertyChain Property;
+	};
+}
\ No newline at end of file
diff --git a/Engine/Plugins/Developer/Concert/ConcertUI/ConcertClientSharedSlate/Source/ConcertClientSharedSlate/Public/Replication/Editor/View/IReplicationEditorView.h b/Engine/Plugins/Developer/Concert/ConcertUI/ConcertClientSharedSlate/Source/ConcertClientSharedSlate/Public/Replication/Editor/View/IReplicationEditorView.h
deleted file mode 100644
diff --git a/Engine/Plugins/Developer/Concert/ConcertUI/ConcertClientSharedSlate/Source/ConcertClientSharedSlate/Public/Replication/Editor/View/IReplicationStreamEditor.h b/Engine/Plugins/Developer/Concert/ConcertUI/ConcertClientSharedSlate/Source/ConcertClientSharedSlate/Public/Replication/Editor/View/IReplicationStreamEditor.h
new file mode 100644
--- /dev/null
+++ b/Engine/Plugins/Developer/Concert/ConcertUI/ConcertClientSharedSlate/Source/ConcertClientSharedSlate/Public/Replication/Editor/View/IReplicationStreamEditor.h
@@ -0,0 +1,16 @@
+// Copyright Epic Games, Inc. All Rights Reserved.
+
+#pragma once
+
+#include "IReplicationStreamViewer.h"
+#include "Templates/SharedPointer.h"
+
+namespace UE::ConcertClientSharedSlate
+{
+	/**
+	 * Widget which edits replication stream.
+	 * @see ReplicationWidgetFactories.h
+	 */
+	class CONCERTCLIENTSHAREDSLATE_API IReplicationStreamEditor : public IReplicationStreamViewer
+	{};
+}
\ No newline at end of file
diff --git a/Engine/Plugins/Developer/Concert/ConcertUI/ConcertClientSharedSlate/Source/ConcertClientSharedSlate/Public/Replication/Editor/View/IReplicationStreamViewer.h b/Engine/Plugins/Developer/Concert/ConcertUI/ConcertClientSharedSlate/Source/ConcertClientSharedSlate/Public/Replication/Editor/View/IReplicationStreamViewer.h
new file mode 100644
--- /dev/null
+++ b/Engine/Plugins/Developer/Concert/ConcertUI/ConcertClientSharedSlate/Source/ConcertClientSharedSlate/Public/Replication/Editor/View/IReplicationStreamViewer.h
@@ -0,0 +1,32 @@
+// Copyright Epic Games, Inc. All Rights Reserved.
+
+#pragma once
+
+#include "Templates/SharedPointer.h"
+#include "Widgets/SCompoundWidget.h"
+
+namespace UE::ConcertClientSharedSlate
+{
+	/**
+	 * Widget which views a replication stream.
+	 * @see ReplicationWidgetFactories.h
+	 */
+	class CONCERTCLIENTSHAREDSLATE_API IReplicationStreamViewer : public SCompoundWidget
+	{
+	public:
+
+		/** Call after the data underlying the model was externally changed and needs to be redisplayed in the UI. */
+		virtual void Refresh() = 0;
+
+		/** @return The objects selected in the top of the viewer; these are the objects that IReplicationSubobjectView bases its view of. */
+		virtual TArray<FSoftObjectPath> GetSelectedTopLevelObjects() const = 0;
+		/**
+		 * @return The objects for which the properties are being edited.
+		 * If there is an IReplicationSubobjectView, this is IReplicationSubobjectView::GetSelectedObjects.
+		 * Otherwise it is IReplicationStreamViewer::GetSelectedTopLevelObjects.
+		 */
+		virtual TArray<FSoftObjectPath> GetObjectsBeingPropertyEdited() const = 0;
+
+		virtual ~IReplicationStreamViewer() = default;
+	};
+}
\ No newline at end of file
diff --git a/Engine/Plugins/Developer/Concert/ConcertUI/ConcertClientSharedSlate/Source/ConcertClientSharedSlate/Public/Replication/Editor/View/PredefinedReplicationColumns.h b/Engine/Plugins/Developer/Concert/ConcertUI/ConcertClientSharedSlate/Source/ConcertClientSharedSlate/Public/Replication/Editor/View/PredefinedReplicationColumns.h
new file mode 100644
--- /dev/null
+++ b/Engine/Plugins/Developer/Concert/ConcertUI/ConcertClientSharedSlate/Source/ConcertClientSharedSlate/Public/Replication/Editor/View/PredefinedReplicationColumns.h
@@ -0,0 +1,59 @@
+// Copyright Epic Games, Inc. All Rights Reserved.
+
+#pragma once
+
+#include "Containers/ArrayView.h"
+#include "Replication/Editor/View/ReplicationColumn.h"
+#include "Replication/Editor/Model/ReplicatedObjectData.h"
+#include "Replication/Editor/Model/ReplicatedPropertyData.h"
+#include "Templates/SharedPointer.h"
+
+struct FConcertPropertyChain;
+
+namespace UE::ConcertClientSharedSlate
+{
+	class IEditableObjectToPropertiesModel;
+	class IReplicationStreamViewer;
+	class IObjectToPropertiesModel;
+}
+
+namespace UE::ConcertClientSharedSlate::ReplicationObjectColumns
+{
+	using FReplicationObjectColumn = TReplicationColumn<FReplicatedObjectData>;
+	enum class EReplicationColumnOrder : int32
+	{
+		/** Displays the class icon */
+		Icon = 10,
+		/** Label of the object */
+		Label = 20,
+		/** Class of the object */
+		Type = 30,
+	};
+}
+
+namespace UE::ConcertClientSharedSlate::ReplicationPropertyColumns
+{
+	using FReplicationPropertyColumn = TReplicationColumn<FReplicatedPropertyData>;
+	enum class EReplicationPropertyColumnOrder : int32
+	{
+		/** The checkbox in SDefaultReplicationStreamEditor determining whether the property is in the selection*/
+		ReplicatesCheckbox = 0,
+		/** Label of the property */
+		Label = 10,
+		/** Type of the property */
+		Type = 20
+	};
+	
+	DECLARE_DELEGATE_RetVal_OneParam(ECheckBoxState, FGetColumnCheckboxState, const FConcertPropertyChain& /*Property*/);
+	DECLARE_DELEGATE_TwoParams(FOnColumnCheckboxChanged, bool /*bIsChecked*/, const FConcertPropertyChain& /*Property*/);
+
+	/** Util for creating a checkbox in a property row. The delegates will be given the path to the property. */
+	CONCERTCLIENTSHAREDSLATE_API FReplicationPropertyColumn MakePropertyCheckboxColumn(
+		FGetColumnCheckboxState GetPropertyCheckboxStateDelegate,
+		FOnColumnCheckboxChanged OnPropertyBoxToggledDelegate,
+		FText DefaultLabel = FText::GetEmpty(),
+		FText ToolTipText = FText::GetEmpty(),
+		const float ColumnWidth = 20.f,
+		const int32 Priority = 0
+		);
+}
diff --git a/Engine/Plugins/Developer/Concert/ConcertUI/ConcertClientSharedSlate/Source/ConcertClientSharedSlate/Public/Replication/Editor/View/ReplicationColumn.h b/Engine/Plugins/Developer/Concert/ConcertUI/ConcertClientSharedSlate/Source/ConcertClientSharedSlate/Public/Replication/Editor/View/ReplicationColumn.h
new file mode 100644
--- /dev/null
+++ b/Engine/Plugins/Developer/Concert/ConcertUI/ConcertClientSharedSlate/Source/ConcertClientSharedSlate/Public/Replication/Editor/View/ReplicationColumn.h
@@ -0,0 +1,72 @@
+// Copyright Epic Games, Inc. All Rights Reserved.
+
+#pragma once
+
+#include "Widgets/SNullWidget.h"
+#include "Widgets/Views/SHeaderRow.h"
+
+namespace UE::ConcertClientSharedSlate
+{
+	template<typename TListItemType>
+	class TReplicationColumn : public SHeaderRow::FColumn
+	{
+	public:
+
+		struct FBuildArgs
+		{
+			TSharedPtr<FText> HighlightText;
+			TListItemType RowData;
+		};
+
+		/** Generates the widget content for a column in a row */
+		DECLARE_DELEGATE_RetVal_OneParam(TSharedRef<SWidget>, FGenerateColumnWidget,
+			const FBuildArgs& InArgs
+			);
+
+		DECLARE_DELEGATE_TwoParams(FPopulateSearchString,
+			const TListItemType& InArgs,
+			TArray<FString>& InOutSearchStrings
+			);
+		
+		SLATE_BEGIN_ARGS(TReplicationColumn)
+		{}
+			/** Callback to use for generating the widget in the column */
+			SLATE_EVENT(FGenerateColumnWidget, GenerateWidgetColumn)
+			/** Callback used to generate search items for this column */
+			SLATE_EVENT(FPopulateSearchString, PopulateSearchItems)
+			/** Where in the row this column is found with respect to the other columns */
+			SLATE_ARGUMENT(int32, ColumnSortOrder)
+		SLATE_END_ARGS()
+
+		TReplicationColumn(const TReplicationColumn::FArguments& InArgs, const FColumn::FArguments& InColumnArgs)
+			: FColumn(InColumnArgs)
+			, GenerateColumnWidgetCallback(InArgs._GenerateWidgetColumn)
+			, PopulateSearchStringCallback(InArgs._PopulateSearchItems)
+			, ColumnSortOrderValue(InArgs._ColumnSortOrder)
+		{}
+
+		TSharedRef<SWidget> BuildColumnWidget(const FBuildArgs& InArgs) const
+		{
+			return ensure(GenerateColumnWidgetCallback.IsBound())
+				? GenerateColumnWidgetCallback.Execute(InArgs)
+				: SNullWidget::NullWidget;
+		}
+
+		void ExecutePopulateSearchString(const TListItemType& InRowData, TArray<FString>& InOutSearchStrings) const
+		{
+			if (PopulateSearchStringCallback.IsBound())
+			{
+				PopulateSearchStringCallback.Execute(InRowData, InOutSearchStrings);
+			}
+		}
+		
+		int32 GetColumnSortOrderValue() const { return ColumnSortOrderValue; }
+		
+	private:
+		
+		FGenerateColumnWidget GenerateColumnWidgetCallback;
+		FPopulateSearchString PopulateSearchStringCallback;
+		/** Determines whether this column is the first, etc. */
+		int32 ColumnSortOrderValue = 0;
+	};
+}
\ No newline at end of file
diff --git a/Engine/Plugins/Developer/Concert/ConcertUI/ConcertClientSharedSlate/Source/ConcertClientSharedSlate/Public/Replication/ReplicationWidgetDelegates.h b/Engine/Plugins/Developer/Concert/ConcertUI/ConcertClientSharedSlate/Source/ConcertClientSharedSlate/Public/Replication/ReplicationWidgetDelegates.h
new file mode 100644
--- /dev/null
+++ b/Engine/Plugins/Developer/Concert/ConcertUI/ConcertClientSharedSlate/Source/ConcertClientSharedSlate/Public/Replication/ReplicationWidgetDelegates.h
@@ -0,0 +1,19 @@
+// Copyright Epic Games, Inc. All Rights Reserved.
+
+#pragma once
+
+#include "Delegates/Delegate.h"
+#include "Delegates/DelegateCombinations.h"
+
+class FMenuBuilder;
+
+namespace UE::ConcertClientSharedSlate
+{
+	class FReplicatedPropertyData;
+
+	/** A predicate for determining Left < Right. */
+	DECLARE_DELEGATE_RetVal_TwoParams(bool, FSortPropertyPredicate, const FReplicatedPropertyData& Left, const FReplicatedPropertyData& Right);
+
+	/** Extends entries already added to a FMenuBuilder */
+	DECLARE_DELEGATE_OneParam(FExtendMenu, FMenuBuilder&);
+}
\ No newline at end of file
diff --git a/Engine/Plugins/Developer/Concert/ConcertUI/ConcertClientSharedSlate/Source/ConcertClientSharedSlate/Public/Replication/ReplicationWidgetFactories.h b/Engine/Plugins/Developer/Concert/ConcertUI/ConcertClientSharedSlate/Source/ConcertClientSharedSlate/Public/Replication/ReplicationWidgetFactories.h
--- a/Engine/Plugins/Developer/Concert/ConcertUI/ConcertClientSharedSlate/Source/ConcertClientSharedSlate/Public/Replication/ReplicationWidgetFactories.h
+++ b/Engine/Plugins/Developer/Concert/ConcertUI/ConcertClientSharedSlate/Source/ConcertClientSharedSlate/Public/Replication/ReplicationWidgetFactories.h
@@ -2,9 +2,21 @@
 
 #pragma once
 
+#include "ReplicationWidgetDelegates.h"
+#include "Editor/View/PredefinedReplicationColumns.h"
 #include "Misc/Attribute.h"
 #include "Templates/SharedPointer.h"
 
+namespace UE::ConcertClientSharedSlate
+{
+	class IObjectToPropertiesModel;
+}
+
+namespace UE::ConcertClientSharedSlate
+{
+	class IReplicationStreamViewer;
+}
+
 class UObject;
 class SWidget;
 
@@ -19,10 +31,31 @@ namespace UE::ConcertClientSharedSlate
 namespace UE::ConcertClientSharedSlate
 {
 	class IEditableObjectToPropertiesModel;
-	class IReplicationEditorView;
+	class IReplicationStreamEditor;
 	class IObjectSelectionSourceModel;
 	class IPropertySelectionSourceModel;
-		
+	
+	/**
+	 * Creates a model that can be passed to CreateEditor.
+	 * 
+	 * This model edits a FObjectReplicationMap that is assumed to be within the transactional OwnerObject.
+	 * The model will respond to undo & redo by triggering the model's update callbacks.
+	 * 
+	 * @param OwnerObject The object containing the FObjectReplicationMap.
+	 * @param ReplicationMapAttribute Getter for extracting the FObjectReplicationMap to edit
+	 * 
+	 * @return A model that will edit the FObjectReplicationMap.
+	 */
+	CONCERTCLIENTSHAREDSLATE_API TSharedRef<IEditableObjectToPropertiesModel> CreatePropertySelectionModel(
+		UObject& OwnerObject,
+		TAttribute<FObjectReplicationMap*> ReplicationMapAttribute,
+		TAttribute<const FConcertReplicationEditorSettings*> OptionalReplicationSettingsAttribute = {}
+		);
+	
+	/** Creates a stream editor with a subobject view that looks like the SSubobjectEditor. */
+	CONCERTCLIENTSHAREDSLATE_API TSharedRef<IReplicationSubobjectView> CreateUnrealEditorSubobjectView();
+
+	/** Params for creating an IReplicationStreamEditor */
 	struct FCreateEditorParams
 	{
 		/**
@@ -30,63 +63,73 @@ namespace UE::ConcertClientSharedSlate
 		 * @note The view will keep a strong reference to this.
 		 */
 		TSharedRef<IEditableObjectToPropertiesModel> DataModel;
-
 		/**
 		 * Determines the objects that can be added to the object list. 
 		 * @note The view will keep a strong reference to this.
 		 */
 		TSharedRef<IObjectSelectionSourceModel> ObjectSource;
-
 		/**
 		 * Determines the properties that can be added to the property list.
 		 * @note The view will keep a strong reference to this.
 		 */
 		TSharedRef<IPropertySelectionSourceModel> PropertySource;
 
-		// TODO DP: Add way to add more columns
-
 		/**
 		 * Optional. This is inserted between the root object outliner and property view.
 		 * It e.g. displays the components of the actor selected in the root object view.
 		 * 
-		 * Exists so it can be customized differently depending on whether used in the editor or on the server-
+		 * Exists so it can be customized differently depending on whether used in the editor or on the server.
+		 * @note The view will keep a strong reference to this.
 		 */
 		TSharedPtr<IReplicationSubobjectView> SubobjectView;
+
+		/** Called to generate the context menu for objects. This extends the options already generated by this widget and is called at the end. */
+		FExtendMenu OnExtendObjectsContextMenu;
+		/** Optional. Used for determining the order in which properties are displayed. */
+		FSortPropertyPredicate SortPropertyRowPredicate;
+		
+		/** Additional columns to add to the object view */
+		TArray<ReplicationObjectColumns::FReplicationObjectColumn> AdditionalObjectColumns;
+		/** Additional columns to add to the property view */
+		TArray<ReplicationPropertyColumns::FReplicationPropertyColumn> AdditionalPropertyColumns;
+		
+		/** Optional widget to add to the left of the object list search bar. */
+		TAlwaysValidWidget LeftOfObjectSearchBar;
+		/** Optional widget to add to the left of the property list search bar. */
+		TAlwaysValidWidget LeftOfPropertySearchBar;
 	};
-	
-	/** Creates an object that looks like the SSubobjectEditor. */
-	CONCERTCLIENTSHAREDSLATE_API TSharedRef<IReplicationSubobjectView> CreateUnrealEditorSubobjectView();
 
 	/**
-	 * Creates a replication editor.
+	 * Creates a base replication stream editor.
 	 * 
-	 * The editor consists of three areas, which share similar workflows as the world outliner and details panel in the level editor:
-	 * 1. Root objects,  similar to world outliner: objects (usually actors) are added here. FCreateEditorParams::ObjectSource is used to build a combo button through which new objects can be added.
-	 * 2. Subobjects (optional), similar to component view (SSubobjectEditor): Shows subobjects of a root objects selected above; typically components.
-	 * 3. Properties, similar to details panel: Shows properties of the selected root object and / or subobjects.
+	 * This editor implements base functionality for adding objects but has NO logic for editing properties.
+	 * You are expected to inject widgets or columns, which implement editing properties.
+	 * @see CreateDefaultStreamEditor, which adds a checkbox to the start of every property row.
+	 *
+	 * The editor consists of three areas, which share similar workflows the level editor:
+	 * 1. Root objects, similar to world outliner:
+	 *		- Similar to World Outliner: Displays top-level objects, i.e. actors, are added here.
+	 *		- FCreateEditorParams::ObjectSource is used to build a combo button through which new objects can be added.
+	 * 2. Subobjects (optional):
+	 *		- Similar to component view (SSubobjectEditor)
+	 *		- Shows subobjects of a root objects selected above; typically components.
+	 * 3. Properties:
+	 *		- Similar to details panel
+	 *		- Shows properties of the selected root object and / or subobjects.
 	 */
-	CONCERTCLIENTSHAREDSLATE_API TSharedRef<IReplicationEditorView> CreateEditor(FCreateEditorParams Params);
+	CONCERTCLIENTSHAREDSLATE_API TSharedRef<IReplicationStreamEditor> CreateEditor(FCreateEditorParams Params);
 
-	/** Creates a default IReplicationEditorView view to use in the Unreal Editor (as opposed to on the server, etc.). */
-	inline TSharedRef<IReplicationEditorView> CreateEditorForUnrealEditor(FCreateEditorParams BaseParams)
-	{
-		BaseParams.SubobjectView = CreateUnrealEditorSubobjectView();
-		return CreateEditor(MoveTemp(BaseParams));
-	}
-	
 	/**
-	 * Creates a model that can be passed to CreateEditor.
-	 * 
-	 * This model edits a FObjectReplicationMap that is assumed to be within the transactional OwnerObject. The model will respond to undo & redo by triggering the model's update callbacks.
-	 * 
-	 * @param OwnerObject The object containing the FObjectReplicationMap.
-	 * @param ReplicationMapAttribute Getter for extracting the FObjectReplicationMap to edit
+	 * Creates a default IReplicationStreamEditor.
+	 *
+	 * This editor adds a checkbox to the start of every property row.
+	 * - Checking adds the property to the selected objects' property mappings.
+	 * - Unchecking removes the property to the selected objects' property mappings
+	 *
+	 * This requires a subobject view. If you do not pass it in, the default view, which looks like the SSubobjectEditor,
+	 * is created.
 	 * 
-	 * @return A model that will edit the FObjectReplicationMap.
+	 * @see CreateViewer for a description of the UI layout.
 	 */
-	CONCERTCLIENTSHAREDSLATE_API TSharedRef<IEditableObjectToPropertiesModel> CreatePropertySelectionModel(
-		UObject& OwnerObject,
-		TAttribute<FObjectReplicationMap*> ReplicationMapAttribute,
-		TAttribute<const FConcertReplicationEditorSettings*> OptionalReplicationSettingsAttribute = {}
-		);
+	CONCERTCLIENTSHAREDSLATE_API TSharedRef<IReplicationStreamEditor> CreateDefaultStreamEditor(FCreateEditorParams Params);
 }
\ No newline at end of file
