Fix default filters not being activated

diff --git a/Engine/Plugins/Developer/Concert/ConcertUI/ConcertClientSharedSlate/Source/ConcertClientSharedSlate/Private/Replication/Editor/View/ObjectViewer/Property/SReplicatedPropertiesView.cpp b/Engine/Plugins/Developer/Concert/ConcertUI/ConcertClientSharedSlate/Source/ConcertClientSharedSlate/Private/Replication/Editor/View/ObjectViewer/Property/SReplicatedPropertiesView.cpp
--- a/Engine/Plugins/Developer/Concert/ConcertUI/ConcertClientSharedSlate/Source/ConcertClientSharedSlate/Private/Replication/Editor/View/ObjectViewer/Property/SReplicatedPropertiesView.cpp
+++ b/Engine/Plugins/Developer/Concert/ConcertUI/ConcertClientSharedSlate/Source/ConcertClientSharedSlate/Private/Replication/Editor/View/ObjectViewer/Property/SReplicatedPropertiesView.cpp
@@ -13,6 +13,44 @@
 
 namespace UE::ConcertClientSharedSlate
 {
+	/** Exposes SetFrontendFilterActive so we can manually enable the default filters */
+	class SReplicationFilterBar : public SBasicFilterBar<TSharedPtr<FReplicatedPropertyData>>
+	{
+		using Super = SBasicFilterBar<TSharedPtr<FReplicatedPropertyData>>;
+	public:
+
+		SLATE_BEGIN_ARGS(SReplicationFilterBar)
+		{}
+			SLATE_EVENT(FOnFilterChanged, OnFilterChanged)
+		SLATE_END_ARGS()
+
+		void Construct(const FArguments& InArgs, TArray<SReplicatedPropertiesView::FFilterRef> AllFilters)
+		{
+			Super::Construct(
+			Super::FArguments()
+				.FilterPillStyle(EFilterPillStyle::Basic)
+				.CustomFilters(MoveTemp(AllFilters))
+				.OnFilterChanged(InArgs._OnFilterChanged)
+				.UseSectionsForCategories(true)
+				);
+		}
+
+		// Expose from SBasicFilterBar
+		using Super::SetFrontendFilterActive;
+
+		void SetFilterVisuallyEnabled(const SReplicatedPropertiesView::FFilterRef& Filter, bool bEnabled)
+		{
+			const TSharedRef<SFilter>* FilterWidget = Filters.FindByPredicate([&Filter](const TSharedRef<SFilter>& FilterWidget)
+			{
+				return FilterWidget->GetFrontendFilter() == Filter;
+			});
+			if (ensure(FilterWidget))
+			{
+				FilterWidget->Get().SetEnabled(bEnabled);
+			}
+		}
+	};
+	
 	void SReplicatedPropertiesView::Construct(const FArguments& InArgs)
 	{
 		const FBuildFilterBarResult Filters = BuildFilterBar();
@@ -60,7 +98,12 @@ namespace UE::ConcertClientSharedSlate
 		// Show all the other filters as enabled (not greyed out: blue) - they will not run their logic since they are inverse.
 		for (const FFilterRef& Filter : Filters.EnabledByDefault)
 		{
+			// Makes it appear on the bar
 			FilterBar->SetFilterCheckState(Filter, ECheckBoxState::Checked);
+			// Functionally makes it affect the search
+			FilterBar->SetFrontendFilterActive(Filter, true);
+			// Visually makes the button blue (i.e. so it looks enabled)
+			FilterBar->SetFilterVisuallyEnabled(Filter, true);
 		}
 	}
 
@@ -118,11 +161,8 @@ namespace UE::ConcertClientSharedSlate
 		AllFilters.Append(Result.DisabledByDefault);
 		AllFilters.Append(Result.EnabledByDefault);
 		
-		FilterBar = SNew(SBasicFilterBar<TSharedPtr<FReplicatedPropertyData>>)
-			.FilterPillStyle(EFilterPillStyle::Basic)
-			.CustomFilters(MoveTemp(AllFilters))
-			.OnFilterChanged(this, &SReplicatedPropertiesView::OnItemsChanged)
-			.UseSectionsForCategories(true);
+		FilterBar = SNew(SReplicationFilterBar, MoveTemp(AllFilters))
+			.OnFilterChanged(this, &SReplicatedPropertiesView::OnItemsChanged);
 		
 		return Result;
 	}
diff --git a/Engine/Plugins/Developer/Concert/ConcertUI/ConcertClientSharedSlate/Source/ConcertClientSharedSlate/Private/Replication/Editor/View/ObjectViewer/Property/SReplicatedPropertiesView.h b/Engine/Plugins/Developer/Concert/ConcertUI/ConcertClientSharedSlate/Source/ConcertClientSharedSlate/Private/Replication/Editor/View/ObjectViewer/Property/SReplicatedPropertiesView.h
--- a/Engine/Plugins/Developer/Concert/ConcertUI/ConcertClientSharedSlate/Source/ConcertClientSharedSlate/Private/Replication/Editor/View/ObjectViewer/Property/SReplicatedPropertiesView.h
+++ b/Engine/Plugins/Developer/Concert/ConcertUI/ConcertClientSharedSlate/Source/ConcertClientSharedSlate/Private/Replication/Editor/View/ObjectViewer/Property/SReplicatedPropertiesView.h
@@ -10,13 +10,15 @@
 namespace UE::ConcertClientSharedSlate
 {
 	class FReplicatedPropertyData;
+	class SReplicationFilterBar;
 	
 	/** Displays a searchable tree view of properties (SReplicationTreeView) and decorates it with a SBasicFilterBar. */
 	class SReplicatedPropertiesView : public SCompoundWidget
 	{
-		using FFilterRef = TSharedRef<FFilterBase<TSharedPtr<FReplicatedPropertyData>>>;
 	public:
 		
+		using FFilterRef = TSharedRef<FFilterBase<TSharedPtr<FReplicatedPropertyData>>>;
+		
 		SLATE_BEGIN_ARGS(SReplicatedPropertiesView)
 		{}
 			/*************** Arguments inherited by SReplicationTreeView ***************/
@@ -52,7 +54,7 @@ namespace UE::ConcertClientSharedSlate
 		/** The tree view displaying the replicated properties */
 		TSharedPtr<SReplicationTreeView<FReplicatedPropertyData>> ReplicatedProperties;
 		/** Displays the active filters*/
-		TSharedPtr<SBasicFilterBar<TSharedPtr<FReplicatedPropertyData>>> FilterBar;
+		TSharedPtr<SReplicationFilterBar> FilterBar;
 
 		struct FBuildFilterBarResult
 		{
