Minor fixes to replication UI:
- Fix not being able to select actor after removing and adding it back
- Fix Upload button getting stuck
- Fix failing ensure after reconnecting to session

diff --git a/Engine/Plugins/Developer/Concert/ConcertApp/MultiUserClient/Source/MultiUserClient/Private/Replication/Submission/SubmissionWorkflow_LocalClient.cpp b/Engine/Plugins/Developer/Concert/ConcertApp/MultiUserClient/Source/MultiUserClient/Private/Replication/Submission/SubmissionWorkflow_LocalClient.cpp
--- a/Engine/Plugins/Developer/Concert/ConcertApp/MultiUserClient/Source/MultiUserClient/Private/Replication/Submission/SubmissionWorkflow_LocalClient.cpp
+++ b/Engine/Plugins/Developer/Concert/ConcertApp/MultiUserClient/Source/MultiUserClient/Private/Replication/Submission/SubmissionWorkflow_LocalClient.cpp
@@ -138,6 +138,7 @@ namespace UE::MultiUserClient
 		{
 			InProgressOperation->EmplaceAuthorityRequestPromise(FSubmitAuthorityChangesRequest{ EAuthoritySubmissionErrorCode::NoChange, *InProgressOperation });
 			InProgressOperation->EmplaceAuthorityResponsePromise(FSubmitAuthorityChangesResponse{ EAuthoritySubmissionErrorCode::NoChange, *InProgressOperation });
+			InProgressOperation.Reset();
 			return;
 		}
 		
diff --git a/Engine/Plugins/Developer/Concert/ConcertSync/ConcertSyncServer/Source/ConcertSyncServer/Private/Replication/ConcertServerReplicationManager_ChangeStream.cpp b/Engine/Plugins/Developer/Concert/ConcertSync/ConcertSyncServer/Source/ConcertSyncServer/Private/Replication/ConcertServerReplicationManager_ChangeStream.cpp
--- a/Engine/Plugins/Developer/Concert/ConcertSync/ConcertSyncServer/Source/ConcertSyncServer/Private/Replication/ConcertServerReplicationManager_ChangeStream.cpp
+++ b/Engine/Plugins/Developer/Concert/ConcertSync/ConcertSyncServer/Source/ConcertSyncServer/Private/Replication/ConcertServerReplicationManager_ChangeStream.cpp
@@ -143,6 +143,8 @@ namespace UE::ConcertSyncServer::Replication
 		const FConcertReplication_ChangeStream_Request& Request,
 		FConcertReplication_ChangeStream_Response& Response)
 	{
+		Response = {};
+		
 		const FGuid SendingClientId = ConcertSessionContext.SourceEndpointId;
 		const TSharedRef<FConcertReplicationClient>* SendingClient = Clients.Find(SendingClientId);
 		if (SendingClient && Private::ShouldAcceptRequest(Request, SendingClient->Get(), AuthorityManager.Get(), Response))
diff --git a/Engine/Plugins/Developer/Concert/ConcertUI/ConcertClientSharedSlate/Source/ConcertClientSharedSlate/Private/Replication/Editor/Model/TransactionalPropertySelectionModel.cpp b/Engine/Plugins/Developer/Concert/ConcertUI/ConcertClientSharedSlate/Source/ConcertClientSharedSlate/Private/Replication/Editor/Model/TransactionalPropertySelectionModel.cpp
--- a/Engine/Plugins/Developer/Concert/ConcertUI/ConcertClientSharedSlate/Source/ConcertClientSharedSlate/Private/Replication/Editor/Model/TransactionalPropertySelectionModel.cpp
+++ b/Engine/Plugins/Developer/Concert/ConcertUI/ConcertClientSharedSlate/Source/ConcertClientSharedSlate/Private/Replication/Editor/Model/TransactionalPropertySelectionModel.cpp
@@ -47,7 +47,7 @@ namespace UE::ConcertClientSharedSlate
 
 	bool FTransactionalPropertySelectionModel::MatchesContext(const FTransactionContext& InContext, const TArray<TPair<UObject*, FTransactionObjectEvent>>& TransactionObjectContexts) const
 	{
-		if (!ensure(OwningObject.IsValid()))
+		if (!OwningObject.IsValid())
 		{
 			return false;
 		}
diff --git a/Engine/Plugins/Developer/Concert/ConcertUI/ConcertClientSharedSlate/Source/ConcertClientSharedSlate/Private/Replication/Editor/View/ObjectEditor/SBaseReplicationStreamEditor.cpp b/Engine/Plugins/Developer/Concert/ConcertUI/ConcertClientSharedSlate/Source/ConcertClientSharedSlate/Private/Replication/Editor/View/ObjectEditor/SBaseReplicationStreamEditor.cpp
--- a/Engine/Plugins/Developer/Concert/ConcertUI/ConcertClientSharedSlate/Source/ConcertClientSharedSlate/Private/Replication/Editor/View/ObjectEditor/SBaseReplicationStreamEditor.cpp
+++ b/Engine/Plugins/Developer/Concert/ConcertUI/ConcertClientSharedSlate/Source/ConcertClientSharedSlate/Private/Replication/Editor/View/ObjectEditor/SBaseReplicationStreamEditor.cpp
@@ -168,6 +168,10 @@ namespace UE::ConcertClientSharedSlate
 		
 		checkf(ObjectAndChildren.Num() >= ObjectsToDelete.Num(), TEXT("Above algorithm is broken."));
 		EditablePropertiesModel->RemoveObjects(ObjectAndChildren);
+
+		// Viewer has an internal cache of the selected object. Setting the displayed object to the cached one is a no-op.
+		// Clear the cache so it displays correctly if the user adds the object back straight after removing it.
+		ReplicationViewer->ClearSubobjectSelection();
 	}
 
 	TSharedPtr<SWidget> SBaseReplicationStreamEditor::OnObjectsContextMenuOpening() const
diff --git a/Engine/Plugins/Developer/Concert/ConcertUI/ConcertClientSharedSlate/Source/ConcertClientSharedSlate/Private/Replication/Editor/View/ObjectViewer/SReplicationStreamViewer.cpp b/Engine/Plugins/Developer/Concert/ConcertUI/ConcertClientSharedSlate/Source/ConcertClientSharedSlate/Private/Replication/Editor/View/ObjectViewer/SReplicationStreamViewer.cpp
--- a/Engine/Plugins/Developer/Concert/ConcertUI/ConcertClientSharedSlate/Source/ConcertClientSharedSlate/Private/Replication/Editor/View/ObjectViewer/SReplicationStreamViewer.cpp
+++ b/Engine/Plugins/Developer/Concert/ConcertUI/ConcertClientSharedSlate/Source/ConcertClientSharedSlate/Private/Replication/Editor/View/ObjectViewer/SReplicationStreamViewer.cpp
@@ -104,6 +104,11 @@ namespace UE::ConcertClientSharedSlate
 		SubobjectAndPropertySection->RefreshPropertyData();
 	}
 
+	void SReplicationStreamViewer::ClearSubobjectSelection()
+	{
+		SubobjectAndPropertySection->ClearSubobjectSelection();
+	}
+
 	TSharedRef<FReplicatedObjectData> SReplicationStreamViewer::AllocateObjectData(FSoftObjectPath ObjectPath)
 	{
 		return MakeShared<FReplicatedObjectData>(MoveTemp(ObjectPath));
diff --git a/Engine/Plugins/Developer/Concert/ConcertUI/ConcertClientSharedSlate/Source/ConcertClientSharedSlate/Private/Replication/Editor/View/ObjectViewer/SReplicationStreamViewer.h b/Engine/Plugins/Developer/Concert/ConcertUI/ConcertClientSharedSlate/Source/ConcertClientSharedSlate/Private/Replication/Editor/View/ObjectViewer/SReplicationStreamViewer.h
--- a/Engine/Plugins/Developer/Concert/ConcertUI/ConcertClientSharedSlate/Source/ConcertClientSharedSlate/Private/Replication/Editor/View/ObjectViewer/SReplicationStreamViewer.h
+++ b/Engine/Plugins/Developer/Concert/ConcertUI/ConcertClientSharedSlate/Source/ConcertClientSharedSlate/Private/Replication/Editor/View/ObjectViewer/SReplicationStreamViewer.h
@@ -80,6 +80,8 @@ namespace UE::ConcertClientSharedSlate
 		void RefreshSubobjectData();
 		void RefreshPropertyData();
 
+		void ClearSubobjectSelection();
+
 		/** @return Gets the root objects selected in the outliner; the subobject view chooses which of these objects (or their subobjects) end up in GetSelectedObjectShowingProperties. */
 		TArray<TSharedPtr<FReplicatedObjectData>> GetSelectedOutlinerObjects() const { return ReplicatedObjects->GetSelectedItems(); }
 		
diff --git a/Engine/Plugins/Developer/Concert/ConcertUI/ConcertClientSharedSlate/Source/ConcertClientSharedSlate/Private/Replication/Editor/View/ObjectViewer/SSubobjectAndPropertySection.cpp b/Engine/Plugins/Developer/Concert/ConcertUI/ConcertClientSharedSlate/Source/ConcertClientSharedSlate/Private/Replication/Editor/View/ObjectViewer/SSubobjectAndPropertySection.cpp
--- a/Engine/Plugins/Developer/Concert/ConcertUI/ConcertClientSharedSlate/Source/ConcertClientSharedSlate/Private/Replication/Editor/View/ObjectViewer/SSubobjectAndPropertySection.cpp
+++ b/Engine/Plugins/Developer/Concert/ConcertUI/ConcertClientSharedSlate/Source/ConcertClientSharedSlate/Private/Replication/Editor/View/ObjectViewer/SSubobjectAndPropertySection.cpp
@@ -42,6 +42,14 @@ namespace UE::ConcertClientSharedSlate
 		}
 	}
 
+	void SSubobjectAndPropertySection::ClearSubobjectSelection() const
+	{
+		if (SubobjectView)
+		{
+			SubobjectView->ClearRootObjects();
+		}
+	}
+
 	void SSubobjectAndPropertySection::RefreshSubobjectData()
 	{
 		if (!SubobjectView)
diff --git a/Engine/Plugins/Developer/Concert/ConcertUI/ConcertClientSharedSlate/Source/ConcertClientSharedSlate/Private/Replication/Editor/View/ObjectViewer/SSubobjectAndPropertySection.h b/Engine/Plugins/Developer/Concert/ConcertUI/ConcertClientSharedSlate/Source/ConcertClientSharedSlate/Private/Replication/Editor/View/ObjectViewer/SSubobjectAndPropertySection.h
--- a/Engine/Plugins/Developer/Concert/ConcertUI/ConcertClientSharedSlate/Source/ConcertClientSharedSlate/Private/Replication/Editor/View/ObjectViewer/SSubobjectAndPropertySection.h
+++ b/Engine/Plugins/Developer/Concert/ConcertUI/ConcertClientSharedSlate/Source/ConcertClientSharedSlate/Private/Replication/Editor/View/ObjectViewer/SSubobjectAndPropertySection.h
@@ -49,6 +49,8 @@ namespace UE::ConcertClientSharedSlate
 
 		/** If there is a subobject view, makes it select the root objects. */
 		void SelectRootObjects() const;
+		/** Clears the subobject selection */
+		void ClearSubobjectSelection() const;
 		
 		void RefreshSubobjectData();
 		void RefreshPropertyData();
diff --git a/Engine/Plugins/Developer/Concert/ConcertUI/ConcertClientSharedSlate/Source/ConcertClientSharedSlate/Private/Replication/Editor/View/SubobjectView/SSubobjectView.cpp b/Engine/Plugins/Developer/Concert/ConcertUI/ConcertClientSharedSlate/Source/ConcertClientSharedSlate/Private/Replication/Editor/View/SubobjectView/SSubobjectView.cpp
--- a/Engine/Plugins/Developer/Concert/ConcertUI/ConcertClientSharedSlate/Source/ConcertClientSharedSlate/Private/Replication/Editor/View/SubobjectView/SSubobjectView.cpp
+++ b/Engine/Plugins/Developer/Concert/ConcertUI/ConcertClientSharedSlate/Source/ConcertClientSharedSlate/Private/Replication/Editor/View/SubobjectView/SSubobjectView.cpp
@@ -114,6 +114,7 @@ namespace UE::ConcertClientSharedSlate
 
 		if (SubobjectModel->GetTopLevelObject().IsNull())
 		{
+			PathToObjectDataCache.Reset();
 			SubobjectTreeView->OnItemsChanged();
 			return;
 		}
