Show message when all properties are filtered out

diff --git a/Engine/Plugins/Developer/Concert/ConcertUI/ConcertClientSharedSlate/Source/ConcertClientSharedSlate/Private/Replication/Editor/View/ObjectViewer/Property/SReplicatedPropertiesView.cpp b/Engine/Plugins/Developer/Concert/ConcertUI/ConcertClientSharedSlate/Source/ConcertClientSharedSlate/Private/Replication/Editor/View/ObjectViewer/Property/SReplicatedPropertiesView.cpp
--- a/Engine/Plugins/Developer/Concert/ConcertUI/ConcertClientSharedSlate/Source/ConcertClientSharedSlate/Private/Replication/Editor/View/ObjectViewer/Property/SReplicatedPropertiesView.cpp
+++ b/Engine/Plugins/Developer/Concert/ConcertUI/ConcertClientSharedSlate/Source/ConcertClientSharedSlate/Private/Replication/Editor/View/ObjectViewer/Property/SReplicatedPropertiesView.cpp
@@ -5,6 +5,7 @@
 #include "Filters/SBasicFilterBar.h"
 #include "PropertyFilter_ByPropertyType.h"
 #include "PropertyFrontendFilter.h"
+#include "Replication/Editor/Model/DisplayUtils.h"
 #include "Replication/Editor/Model/ReplicatedPropertyData.h"
 
 #include "UObject/UnrealType.h"
@@ -53,8 +54,9 @@ namespace UE::ConcertClientSharedSlate
 	
 	void SReplicatedPropertiesView::Construct(const FArguments& InArgs)
 	{
-		const FBuildFilterBarResult Filters = BuildFilterBar();
+		SelectedObjectsAttribute = InArgs._SelectedObjects;
 		
+		const FBuildFilterBarResult Filters = BuildFilterBar();
 		ChildSlot
 		[
 			SAssignNew(ReplicatedProperties, SReplicationTreeView<FReplicatedPropertyData>)
@@ -85,6 +87,23 @@ namespace UE::ConcertClientSharedSlate
 				[
 					FilterBar.ToSharedRef()
 				]
+				.TreeAlternateContent()
+				[
+					SNew(SHorizontalBox)
+					+SHorizontalBox::Slot()
+					.HAlign(HAlign_Center)
+					.VAlign(VAlign_Center)
+					[
+						SNew(STextBlock) 
+						.Text(this, &SReplicatedPropertiesView::GetAllFilteredText)
+					]
+				]
+				.ContentToDisplay_Lambda([this]()
+				{
+					return !SelectedObjectsAttribute.Get().IsEmpty() && ReplicatedProperties->GetFilteredRootItems().IsEmpty()
+						? SReplicationTreeView<FReplicatedPropertyData>::EContent::Custom
+						: SReplicationTreeView<FReplicatedPropertyData>::EContent::TreeView;
+				})
 		];
 
 		// For better UX, hide certain properties by default (e.g. why would you want to replicate bools?)
@@ -171,6 +190,27 @@ namespace UE::ConcertClientSharedSlate
 	{
 		return FilterBar->GetAllActiveFilters()->PassesAllFilters(ReplicatedPropertyData);
 	}
+
+	FText SReplicatedPropertiesView::GetAllFilteredText() const
+	{
+		const TArray<FSoftObjectPath> Objects = SelectedObjectsAttribute.Get();
+		if (Objects.IsEmpty())
+		{
+			return FText::GetEmpty();
+		}
+
+		TSet<FString> Names;
+		Algo::Transform(Objects, Names, [](const FSoftObjectPath& ObjectPath)
+		{
+			return DisplayUtils::GetObjectDisplayText(ObjectPath).ToString();
+		});
+		
+		const FText ObjectsText = FText::FromString(FString::Join(Names, TEXT(", ")));
+		return FText::Format(LOCTEXT("AllFitleredFmt", "All properties filtered for selected {0}|plural(one=object,other=objects): {1}."),
+			Objects.Num(),
+			ObjectsText
+		);
+	}
 }
 
 #undef LOCTEXT_NAMESPACE
\ No newline at end of file
diff --git a/Engine/Plugins/Developer/Concert/ConcertUI/ConcertClientSharedSlate/Source/ConcertClientSharedSlate/Private/Replication/Editor/View/ObjectViewer/Property/SReplicatedPropertiesView.h b/Engine/Plugins/Developer/Concert/ConcertUI/ConcertClientSharedSlate/Source/ConcertClientSharedSlate/Private/Replication/Editor/View/ObjectViewer/Property/SReplicatedPropertiesView.h
--- a/Engine/Plugins/Developer/Concert/ConcertUI/ConcertClientSharedSlate/Source/ConcertClientSharedSlate/Private/Replication/Editor/View/ObjectViewer/Property/SReplicatedPropertiesView.h
+++ b/Engine/Plugins/Developer/Concert/ConcertUI/ConcertClientSharedSlate/Source/ConcertClientSharedSlate/Private/Replication/Editor/View/ObjectViewer/Property/SReplicatedPropertiesView.h
@@ -9,6 +9,7 @@
 
 namespace UE::ConcertClientSharedSlate
 {
+	class FReplicatedObjectData;
 	class FReplicatedPropertyData;
 	class SReplicationFilterBar;
 	
@@ -42,6 +43,9 @@ namespace UE::ConcertClientSharedSlate
 
 			/*************** Own arguments ***************/
 			// Please add new arguments here in the future
+
+			/** Gets the objects being property edited. */
+			SLATE_ATTRIBUTE(TArray<FSoftObjectPath>, SelectedObjects)
 		SLATE_END_ARGS()
 
 		void Construct(const FArguments& InArgs);
@@ -56,6 +60,9 @@ namespace UE::ConcertClientSharedSlate
 		/** Displays the active filters*/
 		TSharedPtr<SReplicationFilterBar> FilterBar;
 
+		/** Used to tell the user that the selected object have all properties filtered out. */
+		TAttribute<TArray<FSoftObjectPath>> SelectedObjectsAttribute;
+
 		struct FBuildFilterBarResult
 		{
 			TArray<FFilterRef> EnabledByDefault;
@@ -66,6 +73,9 @@ namespace UE::ConcertClientSharedSlate
 
 		/** Runs all filters through this item */
 		bool PassesFilters(const TSharedPtr<FReplicatedPropertyData>& ReplicatedPropertyData) const;
+
+		/** Gets the message to display when all properties are filtered out. */
+		FText GetAllFilteredText() const;
 	};
 }
 
diff --git a/Engine/Plugins/Developer/Concert/ConcertUI/ConcertClientSharedSlate/Source/ConcertClientSharedSlate/Private/Replication/Editor/View/ObjectViewer/SReplicationStreamViewer.cpp b/Engine/Plugins/Developer/Concert/ConcertUI/ConcertClientSharedSlate/Source/ConcertClientSharedSlate/Private/Replication/Editor/View/ObjectViewer/SReplicationStreamViewer.cpp
--- a/Engine/Plugins/Developer/Concert/ConcertUI/ConcertClientSharedSlate/Source/ConcertClientSharedSlate/Private/Replication/Editor/View/ObjectViewer/SReplicationStreamViewer.cpp
+++ b/Engine/Plugins/Developer/Concert/ConcertUI/ConcertClientSharedSlate/Source/ConcertClientSharedSlate/Private/Replication/Editor/View/ObjectViewer/SReplicationStreamViewer.cpp
@@ -54,7 +54,7 @@ namespace UE::ConcertClientSharedSlate
 
 	TArray<FSoftObjectPath> SReplicationStreamViewer::GetObjectsBeingPropertyEdited() const
 	{
-		return SubobjectAndPropertySection->GetSelectedObjects();
+		return SubobjectAndPropertySection->GetObjectsSelectedForPropertyEditing();
 	}
 
 	void SReplicationStreamViewer::RefreshObjectData()
diff --git a/Engine/Plugins/Developer/Concert/ConcertUI/ConcertClientSharedSlate/Source/ConcertClientSharedSlate/Private/Replication/Editor/View/ObjectViewer/SSubobjectAndPropertySection.cpp b/Engine/Plugins/Developer/Concert/ConcertUI/ConcertClientSharedSlate/Source/ConcertClientSharedSlate/Private/Replication/Editor/View/ObjectViewer/SSubobjectAndPropertySection.cpp
--- a/Engine/Plugins/Developer/Concert/ConcertUI/ConcertClientSharedSlate/Source/ConcertClientSharedSlate/Private/Replication/Editor/View/ObjectViewer/SSubobjectAndPropertySection.cpp
+++ b/Engine/Plugins/Developer/Concert/ConcertUI/ConcertClientSharedSlate/Source/ConcertClientSharedSlate/Private/Replication/Editor/View/ObjectViewer/SSubobjectAndPropertySection.cpp
@@ -65,7 +65,7 @@ namespace UE::ConcertClientSharedSlate
 	
 	void SSubobjectAndPropertySection::RefreshPropertyData()
 	{
-		const TArray<FSoftObjectPath> SelectedObjects = GetSelectedObjects();
+		const TArray<FSoftObjectPath> SelectedObjects = GetObjectsSelectedForPropertyEditing();
 		if (SelectedObjects.IsEmpty())
 		{
 			SetPropertyContent(EReplicatedPropertyContent::NoSelection);
@@ -113,7 +113,7 @@ namespace UE::ConcertClientSharedSlate
 		ReplicatedProperties->OnItemsChanged();
 	}
 
-	TArray<FSoftObjectPath> SSubobjectAndPropertySection::GetSelectedObjects() const
+	TArray<FSoftObjectPath> SSubobjectAndPropertySection::GetObjectsSelectedForPropertyEditing() const
 	{
 		if (SubobjectView)
 		{
@@ -134,20 +134,37 @@ namespace UE::ConcertClientSharedSlate
 		{
 			SubobjectView = InArgs._SubobjectView;
 			SubobjectView->OnSelectionChanged().AddSP(this, &SSubobjectAndPropertySection::OnSubobjectSelectionChanged);
-			return SNew(SSplitter)
-				.Orientation(Orient_Vertical)
-
-				+SSplitter::Slot()
-				.Value(1.f)
+			return SNew(SWidgetSwitcher)
+				.WidgetIndex_Lambda([this]()
+				{
+					return GetSelectedRootObjectsDelegate.Execute().IsEmpty() ? 1 : 0;
+				})
+			
+				+SWidgetSwitcher::Slot()
 				[
-					SubobjectView.ToSharedRef()
+					SNew(SSplitter)
+					.Orientation(Orient_Vertical)
+
+					+SSplitter::Slot()
+					.Value(1.f)
+					[
+						SubobjectView.ToSharedRef()
+					]
+					
+					+SSplitter::Slot()
+					.Value(2.f)
+					[
+						CreatePropertiesView(InArgs)
+					]
 				]
 			
-				+SSplitter::Slot()
-				.Value(2.f)
+				+SWidgetSwitcher::Slot()
+				.VAlign(VAlign_Center)
+				.HAlign(HAlign_Center)
 				[
-                	CreatePropertiesView(InArgs)
-                ];
+					SNew(STextBlock)
+					.Text(LOCTEXT("NoRootObjects", "Select an object to see selected properties"))
+				];
 		}
 		
 		return CreatePropertiesView(InArgs);
@@ -179,6 +196,7 @@ namespace UE::ConcertClientSharedSlate
 				[
 					InArgs._LeftOfPropertySearchBar.Widget
 				]
+				.SelectedObjects(this, &SSubobjectAndPropertySection::GetObjectsSelectedForPropertyEditing)
 			]
 			
 			// EReplicatedPropertyContent::NoSelection
@@ -187,7 +205,7 @@ namespace UE::ConcertClientSharedSlate
 			.HAlign(HAlign_Center)
 			[
 				SNew(STextBlock)
-				.Text(LOCTEXT("NoSelection", "Select an object to see selected properties"))
+				.Text(LOCTEXT("NoPropertyEditedObjects", "Select an object to see selected properties"))
 			]
 			
 			// EReplicatedPropertyContent::SelectionTooBig
diff --git a/Engine/Plugins/Developer/Concert/ConcertUI/ConcertClientSharedSlate/Source/ConcertClientSharedSlate/Private/Replication/Editor/View/ObjectViewer/SSubobjectAndPropertySection.h b/Engine/Plugins/Developer/Concert/ConcertUI/ConcertClientSharedSlate/Source/ConcertClientSharedSlate/Private/Replication/Editor/View/ObjectViewer/SSubobjectAndPropertySection.h
--- a/Engine/Plugins/Developer/Concert/ConcertUI/ConcertClientSharedSlate/Source/ConcertClientSharedSlate/Private/Replication/Editor/View/ObjectViewer/SSubobjectAndPropertySection.h
+++ b/Engine/Plugins/Developer/Concert/ConcertUI/ConcertClientSharedSlate/Source/ConcertClientSharedSlate/Private/Replication/Editor/View/ObjectViewer/SSubobjectAndPropertySection.h
@@ -56,7 +56,7 @@ namespace UE::ConcertClientSharedSlate
 		void RefreshPropertyData();
 		
 		const TArray<TSharedPtr<FReplicatedPropertyData>>& GetPropertyRowData() const { return PropertyRowData; }
-		TArray<FSoftObjectPath> GetSelectedObjects() const;
+		TArray<FSoftObjectPath> GetObjectsSelectedForPropertyEditing() const;
 
 	private:
 		
diff --git a/Engine/Plugins/Developer/Concert/ConcertUI/ConcertClientSharedSlate/Source/ConcertClientSharedSlate/Private/Replication/Editor/View/ObjectViewer/Tree/SReplicationTreeView.h b/Engine/Plugins/Developer/Concert/ConcertUI/ConcertClientSharedSlate/Source/ConcertClientSharedSlate/Private/Replication/Editor/View/ObjectViewer/Tree/SReplicationTreeView.h
--- a/Engine/Plugins/Developer/Concert/ConcertUI/ConcertClientSharedSlate/Source/ConcertClientSharedSlate/Private/Replication/Editor/View/ObjectViewer/Tree/SReplicationTreeView.h
+++ b/Engine/Plugins/Developer/Concert/ConcertUI/ConcertClientSharedSlate/Source/ConcertClientSharedSlate/Private/Replication/Editor/View/ObjectViewer/Tree/SReplicationTreeView.h
@@ -10,6 +10,7 @@
 #include "Widgets/DeclarativeSyntaxSupport.h"
 #include "Widgets/Input/SSearchBox.h"
 #include "Widgets/Layout/SScrollBox.h"
+#include "Widgets/Layout/SWidgetSwitcher.h"
 #include "Widgets/SBoxPanel.h"
 #include "Widgets/Views/STreeView.h"
 
@@ -34,6 +35,12 @@ namespace UE::ConcertClientSharedSlate
 		DECLARE_DELEGATE_RetVal_OneParam(bool, FCustomFilter, const TSharedPtr<TItemType>& Item);
 		DECLARE_DELEGATE_RetVal_OneParam(bool, FIsSearchableItem, const TSharedPtr<TItemType>& Item);
 
+		enum class EContent
+		{
+			TreeView,
+			Custom
+		};
+
 		SLATE_BEGIN_ARGS(SReplicationTreeView<TItemType>)
 			: _HeaderRowVisibility(EVisibility::Visible)
 			, _SelectionMode(ESelectionMode::Single)
@@ -78,6 +85,11 @@ namespace UE::ConcertClientSharedSlate
 			SLATE_NAMED_SLOT(FArguments, LeftOfSearchBar)
 			/** Optional widget to add between the search bar and the table view (e.g. a SBasicFilterBar). */
 			SLATE_NAMED_SLOT(FArguments, RowBelowSearchBar)
+
+			/** Optional, alternate content to show instead of the tree view when ContentToDisplay returns EContent::Custom. */
+			SLATE_NAMED_SLOT(FArguments, TreeAlternateContent)
+			/** Optional attribute that decides whether TreeAlternateContent is displayed instead of the tree view */
+			SLATE_ATTRIBUTE(EContent, ContentToDisplay)
 		SLATE_END_ARGS()
 
 		void Construct(const FArguments& InArgs)
@@ -141,6 +153,7 @@ namespace UE::ConcertClientSharedSlate
 				[
 					SNew(SScrollBox)
 					+SScrollBox::Slot()
+					.FillSize(1.f)
 					[
 						CreateTreeView(InArgs)
 					]
@@ -168,6 +181,7 @@ namespace UE::ConcertClientSharedSlate
 		}
 		
 		TArray<TSharedPtr<TItemType>> GetSelectedItems() const { return TreeView->GetSelectedItems(); }
+		const TArray<TSharedPtr<TItemType>>& GetFilteredRootItems() const { return FilteredRootItems; }
 		
 		virtual FReply OnKeyDown(const FGeometry& MyGeometry, const FKeyEvent& InKeyEvent) override;
 
@@ -235,15 +249,30 @@ namespace UE::ConcertClientSharedSlate
 			.BorderBackgroundColor(FSlateColor(FLinearColor(0.6, 0.6, 0.6)))
 			.Padding(0)
 			[
-				SAssignNew(TreeView, STreeView<TSharedPtr<TItemType>>)
-				.OnGetChildren(this, &SReplicationTreeView::GetRowChildren)
-				.TreeItemsSource(&FilteredRootItems)
-				.OnGenerateRow(this, &SReplicationTreeView::OnGenerateRowWidget)
-				.OnContextMenuOpening(InArgs._OnContextMenuOpening)
-				.OnSelectionChanged_Lambda([OnSelectionChanged = InArgs._OnSelectionChanged](auto, auto){ OnSelectionChanged.ExecuteIfBound(); })
-				.SelectionMode(InArgs._SelectionMode)
-				.AllowOverscroll(EAllowOverscroll::No)
-				.HeaderRow(CreateHeaderRow(InArgs))
+				SNew(SWidgetSwitcher)
+				.WidgetIndex_Lambda([GetContent = InArgs._ContentToDisplay]()
+				{
+					const bool bHasValue = GetContent.IsSet() || GetContent.IsBound();
+					return bHasValue && GetContent.Get() == EContent::Custom ? 0 : 1; 
+				})
+
+				+SWidgetSwitcher::Slot()
+				[
+					InArgs._TreeAlternateContent.Widget
+				]
+				
+				+SWidgetSwitcher::Slot()
+				[
+					SAssignNew(TreeView, STreeView<TSharedPtr<TItemType>>)
+					.OnGetChildren(this, &SReplicationTreeView::GetRowChildren)
+					.TreeItemsSource(&FilteredRootItems)
+					.OnGenerateRow(this, &SReplicationTreeView::OnGenerateRowWidget)
+					.OnContextMenuOpening(InArgs._OnContextMenuOpening)
+					.OnSelectionChanged_Lambda([OnSelectionChanged = InArgs._OnSelectionChanged](auto, auto){ OnSelectionChanged.ExecuteIfBound(); })
+					.SelectionMode(InArgs._SelectionMode)
+					.AllowOverscroll(EAllowOverscroll::No)
+					.HeaderRow(CreateHeaderRow(InArgs))
+				]
 			];
 	}
 
