Automatically select newly added top-level objects

diff --git a/Engine/Plugins/Developer/Concert/ConcertUI/ConcertClientSharedSlate/Source/ConcertClientSharedSlate/Private/Replication/Editor/View/ObjectEditor/FakeObjectToPropertiesEditorModel.cpp b/Engine/Plugins/Developer/Concert/ConcertUI/ConcertClientSharedSlate/Source/ConcertClientSharedSlate/Private/Replication/Editor/View/ObjectEditor/FakeObjectToPropertiesEditorModel.cpp
--- a/Engine/Plugins/Developer/Concert/ConcertUI/ConcertClientSharedSlate/Source/ConcertClientSharedSlate/Private/Replication/Editor/View/ObjectEditor/FakeObjectToPropertiesEditorModel.cpp
+++ b/Engine/Plugins/Developer/Concert/ConcertUI/ConcertClientSharedSlate/Source/ConcertClientSharedSlate/Private/Replication/Editor/View/ObjectEditor/FakeObjectToPropertiesEditorModel.cpp
@@ -42,6 +42,13 @@ namespace UE::ConcertClientSharedSlate
 		}
 	}
 	
+	bool FFakeObjectToPropertiesEditorModel::IsTopLevelObject(const FSoftObjectPath& ObjectPath) const
+	{
+		const FSoftClassPath ClassPath = GetObjectClass(ObjectPath);
+		const UClass* LoadedClass = ClassPath.TryLoadClass<UObject>();
+		return LoadedClass && LoadedClass->IsChildOf(AActor::StaticClass());
+	}
+	
 	FSoftClassPath FFakeObjectToPropertiesEditorModel::GetObjectClass(const FSoftObjectPath& Object) const
 	{
 		const FSoftClassPath ResolvedClass = RealModel->GetObjectClass(Object);
@@ -57,7 +64,7 @@ namespace UE::ConcertClientSharedSlate
 
 	bool FFakeObjectToPropertiesEditorModel::ForEachReplicatedObject(TFunctionRef<EBreakBehavior(const FSoftObjectPath& Object)> Delegate) const
 	{
-		// This makes SObjectToPropertyViewer only show actors in the outliner
+		// This makes SObjectToPropertyView only show actors in the outliner
 		return IterateTopLevelObjects(Delegate);
 	}
 
@@ -119,13 +126,6 @@ namespace UE::ConcertClientSharedSlate
 		return bResult;
 	}
 
-	bool FFakeObjectToPropertiesEditorModel::IsTopLevelObject(const FSoftObjectPath& ObjectPath) const
-	{
-		const FSoftClassPath ClassPath = GetObjectClass(ObjectPath);
-		const UClass* LoadedClass = ClassPath.TryLoadClass<UObject>();
-		return LoadedClass && LoadedClass->IsChildOf(AActor::StaticClass());
-	}
-
 	bool FFakeObjectToPropertiesEditorModel::IterateDisplayedProperties(const FSoftObjectPath& ObjectPath, TFunctionRef<EBreakBehavior(const FConcertPropertyChain& Property)> Delegate) const
 	{
 		// GetObjectClass might return null because ObjectPath might not be contained in the model
@@ -141,7 +141,7 @@ namespace UE::ConcertClientSharedSlate
 			return false;
 		}
 		
-		// This makes SObjectToPropertyViewer list all available properties
+		// This makes SObjectToPropertyView list all available properties
 		PropertySelectionSource->GetPropertySource(LoadedClass)
 		   ->EnumerateSelectableItems([&Delegate](const FSelectablePropertyInfo& PropertyInfo)
 		   {
diff --git a/Engine/Plugins/Developer/Concert/ConcertUI/ConcertClientSharedSlate/Source/ConcertClientSharedSlate/Private/Replication/Editor/View/ObjectEditor/SBaseReplicationStreamEditor.cpp b/Engine/Plugins/Developer/Concert/ConcertUI/ConcertClientSharedSlate/Source/ConcertClientSharedSlate/Private/Replication/Editor/View/ObjectEditor/SBaseReplicationStreamEditor.cpp
--- a/Engine/Plugins/Developer/Concert/ConcertUI/ConcertClientSharedSlate/Source/ConcertClientSharedSlate/Private/Replication/Editor/View/ObjectEditor/SBaseReplicationStreamEditor.cpp
+++ b/Engine/Plugins/Developer/Concert/ConcertUI/ConcertClientSharedSlate/Source/ConcertClientSharedSlate/Private/Replication/Editor/View/ObjectEditor/SBaseReplicationStreamEditor.cpp
@@ -103,6 +103,17 @@ namespace UE::ConcertClientSharedSlate
 	void SBaseReplicationStreamEditor::OnObjectsChanged(TConstArrayView<UObject*> AddedObjects, TConstArrayView<FSoftObjectPath> RemovedObjects, EReplicatedObjectChangeReason ChangeReason)
 	{
 		ReplicationViewer->RefreshObjectData();
+
+		// Newly added objects should be automatically selected
+		if (!AddedObjects.IsEmpty())
+		{
+			TArray<FSoftObjectPath> TopLevelObjects;
+			Algo::TransformIf(AddedObjects, TopLevelObjects, [this](const UObject* Object)
+			{
+				return PropertiesModelAdapter->IsTopLevelObject(Object);
+			}, [](const UObject* Object){ return Object; });
+			ReplicationViewer->SelectTopLevelObjects(TopLevelObjects);
+		}
 	}
 
 	void SBaseReplicationStreamEditor::OnPropertiesChanged()
diff --git a/Engine/Plugins/Developer/Concert/ConcertUI/ConcertClientSharedSlate/Source/ConcertClientSharedSlate/Private/Replication/Editor/View/ObjectViewer/SReplicationStreamViewer.cpp b/Engine/Plugins/Developer/Concert/ConcertUI/ConcertClientSharedSlate/Source/ConcertClientSharedSlate/Private/Replication/Editor/View/ObjectViewer/SReplicationStreamViewer.cpp
--- a/Engine/Plugins/Developer/Concert/ConcertUI/ConcertClientSharedSlate/Source/ConcertClientSharedSlate/Private/Replication/Editor/View/ObjectViewer/SReplicationStreamViewer.cpp
+++ b/Engine/Plugins/Developer/Concert/ConcertUI/ConcertClientSharedSlate/Source/ConcertClientSharedSlate/Private/Replication/Editor/View/ObjectViewer/SReplicationStreamViewer.cpp
@@ -62,7 +62,7 @@ namespace UE::ConcertClientSharedSlate
 		const int32 NumElements = PropertiesModel->GetNumReplicatedObjects();
 		// Re-using existing instances is tricky: we cannot update the object path in an item because the list view will no detect this change;
 		// list view only looks at the shared ptr address. So the UI will not be refreshed. Since the number of items will be small, just reallocate... 
-		ObjectRowData.Empty(NumElements);
+		AllObjectRowData.Empty(NumElements);
 
 		// Try to re-use old instances by using the old PathToObjectDataCache. This is also done so the expansion states restore correctly in the tree view.
 		TMap<FSoftObjectPath, TSharedPtr<FReplicatedObjectData>> NewPathToObjectDataCache;
@@ -74,7 +74,7 @@ namespace UE::ConcertClientSharedSlate
 		{
 			const TSharedPtr<FReplicatedObjectData>* ExistingItem = PathToObjectDataCache.Find(Path);
 			const TSharedRef<FReplicatedObjectData> Item = ExistingItem ? ExistingItem->ToSharedRef() : AllocateObjectData(Path);
-			ObjectRowData.Emplace(Item);
+			AllObjectRowData.Emplace(Item);
 			NewPathToObjectDataCache.Emplace(Path, Item);
 			return EBreakBehavior::Continue;
 		});
@@ -104,6 +104,21 @@ namespace UE::ConcertClientSharedSlate
 		SubobjectAndPropertySection->RefreshPropertyData();
 	}
 
+	void SReplicationStreamViewer::SelectTopLevelObjects(TConstArrayView<FSoftObjectPath> Objects)
+	{
+		TArray<TSharedPtr<FReplicatedObjectData>> NewSelectedItems; 
+		Algo::TransformIf(AllObjectRowData, NewSelectedItems, [&Objects](const TSharedPtr<FReplicatedObjectData>& ObjectData)
+			{
+				return Objects.Contains(ObjectData->GetObjectPath());
+			},
+			[](const TSharedPtr<FReplicatedObjectData>& ObjectData){ return ObjectData; }
+		);
+		if (!NewSelectedItems.IsEmpty())
+		{
+			ReplicatedObjects->SetSelectedItems(NewSelectedItems, true);
+		}
+	}
+
 	void SReplicationStreamViewer::ClearSubobjectSelection()
 	{
 		SubobjectAndPropertySection->ClearSubobjectSelection();
@@ -202,7 +217,7 @@ namespace UE::ConcertClientSharedSlate
 	void SReplicationStreamViewer::BuildRootObjectRowData()
 	{
 		TSet<TSharedPtr<FReplicatedObjectData>> NonRootNodes;
-		for (const TSharedPtr<FReplicatedObjectData>& Node : ObjectRowData)
+		for (const TSharedPtr<FReplicatedObjectData>& Node : AllObjectRowData)
 		{
 			GetObjectRowChildren(Node, [&NonRootNodes](TSharedPtr<FReplicatedObjectData> Child)
 			{
@@ -212,7 +227,7 @@ namespace UE::ConcertClientSharedSlate
 
 		// Make RootObjectRowData only contain those nodes which were not listed as children 
 		RootObjectRowData.Empty(NonRootNodes.Num());
-		for (const TSharedPtr<FReplicatedObjectData>& Node : ObjectRowData)
+		for (const TSharedPtr<FReplicatedObjectData>& Node : AllObjectRowData)
 		{
 			if (!NonRootNodes.Contains(Node))
 			{
diff --git a/Engine/Plugins/Developer/Concert/ConcertUI/ConcertClientSharedSlate/Source/ConcertClientSharedSlate/Private/Replication/Editor/View/ObjectViewer/SReplicationStreamViewer.h b/Engine/Plugins/Developer/Concert/ConcertUI/ConcertClientSharedSlate/Source/ConcertClientSharedSlate/Private/Replication/Editor/View/ObjectViewer/SReplicationStreamViewer.h
--- a/Engine/Plugins/Developer/Concert/ConcertUI/ConcertClientSharedSlate/Source/ConcertClientSharedSlate/Private/Replication/Editor/View/ObjectViewer/SReplicationStreamViewer.h
+++ b/Engine/Plugins/Developer/Concert/ConcertUI/ConcertClientSharedSlate/Source/ConcertClientSharedSlate/Private/Replication/Editor/View/ObjectViewer/SReplicationStreamViewer.h
@@ -80,6 +80,10 @@ namespace UE::ConcertClientSharedSlate
 		void RefreshSubobjectData();
 		void RefreshPropertyData();
 
+		/** Selects the given objects from the top level view, if applicable. */
+		void SelectTopLevelObjects(TConstArrayView<FSoftObjectPath> Objects);
+		
+		/** Clears all objects selected in the subobject view, if there is one. */
 		void ClearSubobjectSelection();
 
 		/** @return Gets the root objects selected in the outliner; the subobject view chooses which of these objects (or their subobjects) end up in GetSelectedObjectShowingProperties. */
@@ -98,12 +102,8 @@ namespace UE::ConcertClientSharedSlate
 		/** Tree view for replicated objects. */
 		TSharedPtr<SReplicationTreeView<FReplicatedObjectData>> ReplicatedObjects;
 		
-		/**
-		 * These instances can be subclasses of FReplicatedObjectData, e.g. FReplicatedObjectData_Editor.
-		 * Their type can be overridden by subclasses.
-		 * @see GetObjectData
-		 */
-		TArray<TSharedPtr<FReplicatedObjectData>> ObjectRowData;
+		/** All object row data */
+		TArray<TSharedPtr<FReplicatedObjectData>> AllObjectRowData;
 		/** The instances of ObjectRowData which do not have any parents. This acts as the item source for the tree view. */
 		TArray<TSharedPtr<FReplicatedObjectData>> RootObjectRowData;
 		/** Inverse map of ObjectRowData using FReplicatedObjectData::GetObjectPath as key. Contains all elements of ObjectRowData. */
