Promt user to add replicated objects when there are none

diff --git a/Engine/Plugins/Developer/Concert/ConcertUI/ConcertClientSharedSlate/Source/ConcertClientSharedSlate/Private/Replication/Editor/View/ObjectEditor/SBaseReplicationStreamEditor.cpp b/Engine/Plugins/Developer/Concert/ConcertUI/ConcertClientSharedSlate/Source/ConcertClientSharedSlate/Private/Replication/Editor/View/ObjectEditor/SBaseReplicationStreamEditor.cpp
--- a/Engine/Plugins/Developer/Concert/ConcertUI/ConcertClientSharedSlate/Source/ConcertClientSharedSlate/Private/Replication/Editor/View/ObjectEditor/SBaseReplicationStreamEditor.cpp
+++ b/Engine/Plugins/Developer/Concert/ConcertUI/ConcertClientSharedSlate/Source/ConcertClientSharedSlate/Private/Replication/Editor/View/ObjectEditor/SBaseReplicationStreamEditor.cpp
@@ -18,7 +18,7 @@
 #include "UObject/Class.h"
 #include "Widgets/SBoxPanel.h"
 
-#define LOCTEXT_NAMESPACE "SObjectToPropertyEditor"
+#define LOCTEXT_NAMESPACE "SBaseReplicationStreamEditor"
 
 namespace UE::ConcertClientSharedSlate
 {
@@ -67,6 +67,7 @@ namespace UE::ConcertClientSharedSlate
 				[
 					InArgs._LeftOfPropertySearchBar.Widget
 				]
+				.NoOutlinerObjects(LOCTEXT("NoObjects", "Add objects to replicate"))
 		];
 	}
 
diff --git a/Engine/Plugins/Developer/Concert/ConcertUI/ConcertClientSharedSlate/Source/ConcertClientSharedSlate/Private/Replication/Editor/View/ObjectViewer/Property/SReplicatedPropertiesView.cpp b/Engine/Plugins/Developer/Concert/ConcertUI/ConcertClientSharedSlate/Source/ConcertClientSharedSlate/Private/Replication/Editor/View/ObjectViewer/Property/SReplicatedPropertiesView.cpp
--- a/Engine/Plugins/Developer/Concert/ConcertUI/ConcertClientSharedSlate/Source/ConcertClientSharedSlate/Private/Replication/Editor/View/ObjectViewer/Property/SReplicatedPropertiesView.cpp
+++ b/Engine/Plugins/Developer/Concert/ConcertUI/ConcertClientSharedSlate/Source/ConcertClientSharedSlate/Private/Replication/Editor/View/ObjectViewer/Property/SReplicatedPropertiesView.cpp
@@ -83,27 +83,8 @@ namespace UE::ConcertClientSharedSlate
 						InArgs._LeftOfSearchBar.Widget
 					]
 				]
-				.RowBelowSearchBar()
-				[
-					FilterBar.ToSharedRef()
-				]
-				.TreeAlternateContent()
-				[
-					SNew(SHorizontalBox)
-					+SHorizontalBox::Slot()
-					.HAlign(HAlign_Center)
-					.VAlign(VAlign_Center)
-					[
-						SNew(STextBlock) 
-						.Text(this, &SReplicatedPropertiesView::GetAllFilteredText)
-					]
-				]
-				.ContentToDisplay_Lambda([this]()
-				{
-					return !SelectedObjectsAttribute.Get().IsEmpty() && ReplicatedProperties->GetFilteredRootItems().IsEmpty()
-						? SReplicationTreeView<FReplicatedPropertyData>::EContent::Custom
-						: SReplicationTreeView<FReplicatedPropertyData>::EContent::TreeView;
-				})
+				.RowBelowSearchBar() [ FilterBar.ToSharedRef() ]
+				.NoItemsContent() [ SNew(STextBlock).Text(this, &SReplicatedPropertiesView::GetAllFilteredText) ]
 		];
 
 		// For better UX, hide certain properties by default (e.g. why would you want to replicate bools?)
diff --git a/Engine/Plugins/Developer/Concert/ConcertUI/ConcertClientSharedSlate/Source/ConcertClientSharedSlate/Private/Replication/Editor/View/ObjectViewer/SReplicationStreamViewer.cpp b/Engine/Plugins/Developer/Concert/ConcertUI/ConcertClientSharedSlate/Source/ConcertClientSharedSlate/Private/Replication/Editor/View/ObjectViewer/SReplicationStreamViewer.cpp
--- a/Engine/Plugins/Developer/Concert/ConcertUI/ConcertClientSharedSlate/Source/ConcertClientSharedSlate/Private/Replication/Editor/View/ObjectViewer/SReplicationStreamViewer.cpp
+++ b/Engine/Plugins/Developer/Concert/ConcertUI/ConcertClientSharedSlate/Source/ConcertClientSharedSlate/Private/Replication/Editor/View/ObjectViewer/SReplicationStreamViewer.cpp
@@ -5,13 +5,13 @@
 #include "ConcertFrontendUtils.h"
 #include "Replication/Editor/Model/IObjectToPropertiesModel.h"
 #include "Replication/Editor/Model/ReplicatedObjectData.h"
-#include "Replication/Editor/Model/ReplicatedPropertyData.h"
 #include "Replication/Editor/View/ObjectViewer/Property/SReplicatedPropertiesView.h"
 #include "Replication/Editor/View/ObjectViewer/Tree/SelectionViewerColumns.h"
 #include "SSubobjectAndPropertySection.h"
 
 #include "Widgets/Layout/SBorder.h"
 #include "Widgets/Layout/SSplitter.h"
+#include "Widgets/Layout/SWidgetSwitcher.h"
 #include "Widgets/SBoxPanel.h"
 #include "Widgets/Text/STextBlock.h"
 
@@ -137,7 +137,7 @@ namespace UE::ConcertClientSharedSlate
 			+SSplitter::Slot()
 			 .Value(1.f)
 			[
-				CreateActorsSection(InArgs)
+				CreateOutlinerSection(InArgs)
 			]
 
 			+SSplitter::Slot()
@@ -148,7 +148,7 @@ namespace UE::ConcertClientSharedSlate
 			];
 	}
 
-	TSharedRef<SWidget> SReplicationStreamViewer::CreateActorsSection(const FArguments& InArgs)
+	TSharedRef<SWidget> SReplicationStreamViewer::CreateOutlinerSection(const FArguments& InArgs)
 	{
 		TArray Columns
 		{
@@ -158,6 +158,9 @@ namespace UE::ConcertClientSharedSlate
 		};
 		Columns.Append(InArgs._AdditionalObjectColumns);
 		
+		const bool bHasNoOutlinerObjectsAttribute = InArgs._NoOutlinerObjects.IsBound() || InArgs._NoOutlinerObjects.IsSet(); 
+		const TAttribute<FText> NoObjectsAttribute = bHasNoOutlinerObjectsAttribute ? InArgs._NoOutlinerObjects : LOCTEXT("NoObjects", "No objects to display");
+
 		return SAssignNew(ReplicatedObjects, SReplicationTreeView<FReplicatedObjectData>)
 			.RootItemsSource(&RootObjectRowData)
 			.OnGetChildren(this, &SReplicationStreamViewer::GetObjectRowChildren)
@@ -172,10 +175,8 @@ namespace UE::ConcertClientSharedSlate
 			.Columns(Columns)
 			.ExpandableColumnLabel(ReplicationColumns::TopLevel::LabelColumnId)
 			.SelectionMode(ESelectionMode::Multi)
-			.LeftOfSearchBar()
-			[
-				InArgs._LeftOfObjectSearchBar.Widget
-			];
+			.LeftOfSearchBar() [ InArgs._LeftOfObjectSearchBar.Widget ]
+			.NoItemsContent() [ SNew(STextBlock).Text(NoObjectsAttribute) ];
 	}
 
 	TSharedRef<SWidget> SReplicationStreamViewer::CreatePropertiesSection(const FArguments& InArgs)
diff --git a/Engine/Plugins/Developer/Concert/ConcertUI/ConcertClientSharedSlate/Source/ConcertClientSharedSlate/Private/Replication/Editor/View/ObjectViewer/SReplicationStreamViewer.h b/Engine/Plugins/Developer/Concert/ConcertUI/ConcertClientSharedSlate/Source/ConcertClientSharedSlate/Private/Replication/Editor/View/ObjectViewer/SReplicationStreamViewer.h
--- a/Engine/Plugins/Developer/Concert/ConcertUI/ConcertClientSharedSlate/Source/ConcertClientSharedSlate/Private/Replication/Editor/View/ObjectViewer/SReplicationStreamViewer.h
+++ b/Engine/Plugins/Developer/Concert/ConcertUI/ConcertClientSharedSlate/Source/ConcertClientSharedSlate/Private/Replication/Editor/View/ObjectViewer/SReplicationStreamViewer.h
@@ -66,6 +66,9 @@ namespace UE::ConcertClientSharedSlate
 			SLATE_NAMED_SLOT(FArguments, LeftOfObjectSearchBar)
 			/** Optional widget to add to the left of the property list search bar. */
 			SLATE_NAMED_SLOT(FArguments, LeftOfPropertySearchBar)
+
+			/** Optional text to display when no object is in the outliner. Defaults to "No objects to display." "*/
+			SLATE_ATTRIBUTE(FText, NoOutlinerObjects)
 		SLATE_END_ARGS()
 
 		void Construct(const FArguments& InArgs, TSharedRef<IObjectToPropertiesModel> InPropertiesModel);
@@ -115,7 +118,7 @@ namespace UE::ConcertClientSharedSlate
 
 		// Widget creation helpers
 		TSharedRef<SWidget> CreateContentWidget(const FArguments& InArgs);
-		TSharedRef<SWidget> CreateActorsSection(const FArguments& InArgs);
+		TSharedRef<SWidget> CreateOutlinerSection(const FArguments& InArgs);
 		TSharedRef<SWidget> CreatePropertiesSection(const FArguments& InArgs);
 
 		/** Sets RootObjectRowData to all non-root nodes from ObjectRowData. */
diff --git a/Engine/Plugins/Developer/Concert/ConcertUI/ConcertClientSharedSlate/Source/ConcertClientSharedSlate/Private/Replication/Editor/View/ObjectViewer/Tree/SReplicationTreeView.h b/Engine/Plugins/Developer/Concert/ConcertUI/ConcertClientSharedSlate/Source/ConcertClientSharedSlate/Private/Replication/Editor/View/ObjectViewer/Tree/SReplicationTreeView.h
--- a/Engine/Plugins/Developer/Concert/ConcertUI/ConcertClientSharedSlate/Source/ConcertClientSharedSlate/Private/Replication/Editor/View/ObjectViewer/Tree/SReplicationTreeView.h
+++ b/Engine/Plugins/Developer/Concert/ConcertUI/ConcertClientSharedSlate/Source/ConcertClientSharedSlate/Private/Replication/Editor/View/ObjectViewer/Tree/SReplicationTreeView.h
@@ -32,6 +32,7 @@ namespace UE::ConcertClientSharedSlate
 		DECLARE_DELEGATE_OneParam(FDeleteItems, const TArray<TSharedPtr<TItemType>>& SelectedItems);
 		DECLARE_DELEGATE_TwoParams(FGetItemChildren, TSharedPtr<TItemType> Item, TFunctionRef<void(TSharedPtr<TItemType>)> ProcessChild);
 		DECLARE_DELEGATE(FOnSelectionChanged);
+		
 		DECLARE_DELEGATE_RetVal_OneParam(bool, FCustomFilter, const TSharedPtr<TItemType>& Item);
 		DECLARE_DELEGATE_RetVal_OneParam(bool, FIsSearchableItem, const TSharedPtr<TItemType>& Item);
 
@@ -86,10 +87,8 @@ namespace UE::ConcertClientSharedSlate
 			/** Optional widget to add between the search bar and the table view (e.g. a SBasicFilterBar). */
 			SLATE_NAMED_SLOT(FArguments, RowBelowSearchBar)
 
-			/** Optional, alternate content to show instead of the tree view when ContentToDisplay returns EContent::Custom. */
-			SLATE_NAMED_SLOT(FArguments, TreeAlternateContent)
-			/** Optional attribute that decides whether TreeAlternateContent is displayed instead of the tree view */
-			SLATE_ATTRIBUTE(EContent, ContentToDisplay)
+			/** Optional, alternate content to show instead of the tree view when there are no rows. */
+			SLATE_NAMED_SLOT(FArguments, NoItemsContent)
 		SLATE_END_ARGS()
 
 		void Construct(const FArguments& InArgs)
@@ -244,24 +243,17 @@ namespace UE::ConcertClientSharedSlate
 	template <typename TItemType>
 	TSharedRef<SWidget> SReplicationTreeView<TItemType>::CreateTreeView(const FArguments& InArgs)
 	{
-		return SNew(SBorder)
+		TSharedPtr<SVerticalBox> VerticalBox;
+		
+		TSharedRef<SWidget> Result = SNew(SBorder)
 			.BorderImage(FAppStyle::Get().GetBrush("ToolPanel.GroupBorder"))
 			.BorderBackgroundColor(FSlateColor(FLinearColor(0.6, 0.6, 0.6)))
 			.Padding(0)
 			[
-				SNew(SWidgetSwitcher)
-				.WidgetIndex_Lambda([GetContent = InArgs._ContentToDisplay]()
-				{
-					const bool bHasValue = GetContent.IsSet() || GetContent.IsBound();
-					return bHasValue && GetContent.Get() == EContent::Custom ? 0 : 1; 
-				})
-
-				+SWidgetSwitcher::Slot()
-				[
-					InArgs._TreeAlternateContent.Widget
-				]
+				SAssignNew(VerticalBox, SVerticalBox)
 				
-				+SWidgetSwitcher::Slot()
+				+SVerticalBox::Slot()
+				.FillHeight(1.f)
 				[
 					SAssignNew(TreeView, STreeView<TSharedPtr<TItemType>>)
 					.OnGetChildren(this, &SReplicationTreeView::GetRowChildren)
@@ -273,7 +265,25 @@ namespace UE::ConcertClientSharedSlate
 					.AllowOverscroll(EAllowOverscroll::No)
 					.HeaderRow(CreateHeaderRow(InArgs))
 				]
+
+				+SVerticalBox::Slot()
+				.FillHeight(1.f)
+				.HAlign(HAlign_Center)
+				.VAlign(VAlign_Center)
+				.Padding(0.f, 0.f, 0.f, 20.f)
+				[
+					SNew(SWidgetSwitcher)
+					.WidgetIndex_Lambda([this](){ return AllRootItems->IsEmpty() ? 1 : 0; })
+					.Visibility_Lambda([this](){ return AllRootItems->IsEmpty() || FilteredRootItems.IsEmpty() ? EVisibility::Visible : EVisibility::Collapsed; })
+					+SWidgetSwitcher::Slot() [ SNew(STextBlock).Text(LOCTEXT("AllFiltered", "All items are filtered.")) ]
+					+SWidgetSwitcher::Slot()
+					[
+						InArgs._NoItemsContent.Widget
+					]
+				]
 			];
+		
+		return Result;
 	}
 
 	template <typename TItemType>
