Fix compile issue on Mac

diff --git a/Engine/Plugins/Developer/Concert/ConcertUI/ConcertClientSharedSlate/Source/ConcertClientSharedSlate/Private/Widgets/ClientName/SHorizontalClientList.cpp b/Engine/Plugins/Developer/Concert/ConcertUI/ConcertClientSharedSlate/Source/ConcertClientSharedSlate/Private/Widgets/ClientName/SHorizontalClientList.cpp
--- a/Engine/Plugins/Developer/Concert/ConcertUI/ConcertClientSharedSlate/Source/ConcertClientSharedSlate/Private/Widgets/ClientName/SHorizontalClientList.cpp
+++ b/Engine/Plugins/Developer/Concert/ConcertUI/ConcertClientSharedSlate/Source/ConcertClientSharedSlate/Private/Widgets/ClientName/SHorizontalClientList.cpp
@@ -80,7 +80,7 @@ namespace UE::ConcertClientSharedSlate
 			}
 			else if (Client == ClientSession->GetSessionClientEndpointId())
 			{
-				ClientsToDisplay.Emplace(Client, ClientSession->GetLocalClientInfo());
+				ClientsToDisplay.Add({Client, ClientSession->GetLocalClientInfo()});
 			}
 		}
 
