Default properties and objects are now only added when a new actor is added

diff --git a/Engine/Plugins/Developer/Concert/ConcertApp/MultiUserClient/Source/MultiUserClient/Private/Replication/Client/ReplicationClient.cpp b/Engine/Plugins/Developer/Concert/ConcertApp/MultiUserClient/Source/MultiUserClient/Private/Replication/Client/ReplicationClient.cpp
--- a/Engine/Plugins/Developer/Concert/ConcertApp/MultiUserClient/Source/MultiUserClient/Private/Replication/Client/ReplicationClient.cpp
+++ b/Engine/Plugins/Developer/Concert/ConcertApp/MultiUserClient/Source/MultiUserClient/Private/Replication/Client/ReplicationClient.cpp
@@ -2,7 +2,6 @@
 
 #include "ReplicationClient.h"
 
-#include "MultiUserReplicationSettings.h"
 #include "Assets/MultiUserReplicationClientPreset.h"
 #include "Replication/Editor/Model/IEditableObjectToPropertiesModel.h"
 #include "Replication/ReplicationWidgetFactories.h"
@@ -23,8 +22,7 @@ namespace UE::MultiUserClient
 		, AuthoritySynchronizer(MoveTemp(InAuthoritySynchronizer))
 		, LocalClientEditModel(ConcertClientSharedSlate::CreatePropertySelectionModel(
 			*ClientContentStorage->Stream,
-			ClientContentStorage->Stream->MakeReplicationMapGetterAttribute(),
-			TAttribute<const FConcertReplicationEditorSettings*>::CreateLambda([](){ return &UMultiUserReplicationSettings::Get()->ReplicationEditorSettings; })
+			ClientContentStorage->Stream->MakeReplicationMapGetterAttribute()
 			))
 		, LocalClientStreamDiffer(
 			GetStreamSynchronizer(),
diff --git a/Engine/Plugins/Developer/Concert/ConcertApp/MultiUserClient/Source/MultiUserClient/Private/Widgets/ActiveSession/Replication/Client/Single/SReplicationClientView.cpp b/Engine/Plugins/Developer/Concert/ConcertApp/MultiUserClient/Source/MultiUserClient/Private/Widgets/ActiveSession/Replication/Client/Single/SReplicationClientView.cpp
--- a/Engine/Plugins/Developer/Concert/ConcertApp/MultiUserClient/Source/MultiUserClient/Private/Widgets/ActiveSession/Replication/Client/Single/SReplicationClientView.cpp
+++ b/Engine/Plugins/Developer/Concert/ConcertApp/MultiUserClient/Source/MultiUserClient/Private/Widgets/ActiveSession/Replication/Client/Single/SReplicationClientView.cpp
@@ -2,6 +2,7 @@
 
 #include "SReplicationClientView.h"
 
+#include "MultiUserReplicationSettings.h"
 #include "SingleClientColumns.h"
 #include "Replication/Client/ReplicationClient.h"
 #include "Replication/Editor/Model/IEditableObjectToPropertiesModel.h"
@@ -29,8 +30,11 @@ namespace UE::MultiUserClient
 		FAuthorityChangeTracker& AuthorityTracker = ReplicationClient->GetAuthorityDiffer();
 		ISubmissionWorkflow& SubmissionWorkflow = ReplicationClient->GetSubmissionWorkflow();
 		FGlobalAuthorityCache& AuthorityCache = InClientManager.GetAuthorityCache();
+		
 		TAttribute<const ConcertClientSharedSlate::IReplicationStreamViewer*> GetReplicationViewerAttribute =
 			TAttribute<const ConcertClientSharedSlate::IReplicationStreamViewer*>::CreateLambda([this](){ return EditorView.Get(); });
+		TAttribute<const FConcertReplicationEditorSettings*> ReplicationSettingsAttribute =
+			TAttribute<const FConcertReplicationEditorSettings*>::CreateLambda([](){ return &UMultiUserReplicationSettings::Get()->ReplicationEditorSettings; });
 
 		// Add checkboxes in front of top level and subobject rows for changing authority
 		using namespace ConcertClientSharedSlate;
@@ -55,10 +59,11 @@ namespace UE::MultiUserClient
 			},
 			.AdditionalPropertyColumns =
 			{
-				SingleClientColumns::OwnerOfProperty(InClient, AuthorityCache, GetReplicationViewerAttribute)
+				SingleClientColumns::OwnerOfProperty(InClient, AuthorityCache, MoveTemp(GetReplicationViewerAttribute))
 			},
 			.IsEditingEnabled = TAttribute<bool>::CreateLambda([&SubmissionWorkflow](){ return SubmissionWorkflow.GetUploadability() != EChangeUploadability::NotImplemented; }),
-			.EditingDisabledToolTipText = LOCTEXT("Editing.NotImplemented", "Editing remote clients is not implemented. You can only edit the local client.")
+			.EditingDisabledToolTipText = LOCTEXT("Editing.NotImplemented", "Editing remote clients is not implemented. You can only edit the local client."),
+			.ReplicationSettingsAttribute = MoveTemp(ReplicationSettingsAttribute)
 		};
 
 		EditorView = CreateDefaultStreamEditor(ReplicationEditorCreationParams);
diff --git a/Engine/Plugins/Developer/Concert/ConcertApp/MultiUserClient/Source/MultiUserReplicationEditor/Private/StreamEditor/ReplicationSessionPresetEditorToolkit.cpp b/Engine/Plugins/Developer/Concert/ConcertApp/MultiUserClient/Source/MultiUserReplicationEditor/Private/StreamEditor/ReplicationSessionPresetEditorToolkit.cpp
--- a/Engine/Plugins/Developer/Concert/ConcertApp/MultiUserClient/Source/MultiUserReplicationEditor/Private/StreamEditor/ReplicationSessionPresetEditorToolkit.cpp
+++ b/Engine/Plugins/Developer/Concert/ConcertApp/MultiUserClient/Source/MultiUserReplicationEditor/Private/StreamEditor/ReplicationSessionPresetEditorToolkit.cpp
@@ -74,14 +74,10 @@ namespace UE::MultiUserReplicationEditor
 	{
 		using namespace ConcertClientSharedSlate;
 		
-		// TODO DP: Create a new model that combines several clients.
 		const TSharedRef<IEditableObjectToPropertiesModel> AssetReadWriteModel = CreatePropertySelectionModel(
 			*GetEditedStreamAsset(),
-			GetEditedStreamAsset()->GetUnassignedClient()->Stream->MakeReplicationMapGetterAttribute(),
-			TAttribute<const FConcertReplicationEditorSettings*>::CreateLambda([]()
-			{
-				return nullptr;
-			}));
+			GetEditedStreamAsset()->GetUnassignedClient()->Stream->MakeReplicationMapGetterAttribute()
+			);
 		const TSharedRef<FEditorObjectSelectionSourceModel> ObjectSourceModel = MakeShared<FEditorObjectSelectionSourceModel>();
 		const TSharedRef<FSelectPropertyFromUClassModel> PropertySourceModel = MakeShared<FSelectPropertyFromUClassModel>();
 		const TSharedRef<IReplicationStreamEditor> EditorView = CreateDefaultStreamEditor(FCreateEditorParams{ AssetReadWriteModel, ObjectSourceModel, PropertySourceModel });
diff --git a/Engine/Plugins/Developer/Concert/ConcertUI/ConcertClientSharedSlate/Source/ConcertClientSharedSlate/Private/Replication/Editor/Model/GenericPropertySelectionModel.cpp b/Engine/Plugins/Developer/Concert/ConcertUI/ConcertClientSharedSlate/Source/ConcertClientSharedSlate/Private/Replication/Editor/Model/GenericPropertySelectionModel.cpp
--- a/Engine/Plugins/Developer/Concert/ConcertUI/ConcertClientSharedSlate/Source/ConcertClientSharedSlate/Private/Replication/Editor/Model/GenericPropertySelectionModel.cpp
+++ b/Engine/Plugins/Developer/Concert/ConcertUI/ConcertClientSharedSlate/Source/ConcertClientSharedSlate/Private/Replication/Editor/Model/GenericPropertySelectionModel.cpp
@@ -10,12 +10,8 @@
 
 namespace UE::ConcertClientSharedSlate
 {
-	FGenericPropertySelectionModel::FGenericPropertySelectionModel(
-		TAttribute<FObjectReplicationMap*> ReplicationMapAttribute,
-		TAttribute<const FConcertReplicationEditorSettings*> OptionalReplicationSettingsAttribute
-		)
+	FGenericPropertySelectionModel::FGenericPropertySelectionModel(TAttribute<FObjectReplicationMap*> ReplicationMapAttribute)
 		: ReplicationMapAttribute(MoveTemp(ReplicationMapAttribute))
-		, OptionalReplicationSettingsAttribute(MoveTemp(OptionalReplicationSettingsAttribute))
 	{}
 	
 	uint32 FGenericPropertySelectionModel::GetNumReplicatedObjects() const
@@ -126,7 +122,6 @@ namespace UE::ConcertClientSharedSlate
 		}
 		
 		TSet<UObject*> ObjectsNotAdded;
-		TArray<UObject*> AdditionalObjectsToAdd;
 		for (UObject* Object : Objects)
 		{
 			const FSoftObjectPath ObjectPath = Object;
@@ -134,20 +129,6 @@ namespace UE::ConcertClientSharedSlate
 			{
 				FReplicatedObjectInfo& ObjectInfo = ReplicationMap->ReplicatedObjects.Add(ObjectPath);
 				ObjectInfo.ClassPath = Object->GetClass();
-
-				const FConcertReplicationEditorSettings* AutoPopulateSettings = OptionalReplicationSettingsAttribute.IsBound()
-					? OptionalReplicationSettingsAttribute.Get()
-					: nullptr;
-				if (!AutoPopulateSettings)
-				{
-					continue;
-				}
-				
-				AutoPopulateSettings->AddDefaultPropertiesFromSettings(ObjectInfo, *Object->GetClass());
-				AutoPopulateSettings->AddAdditionalObjectsFromSettings(*Object, [&AdditionalObjectsToAdd](UObject& FurtherObject)
-				{
-					AdditionalObjectsToAdd.Add(&FurtherObject);
-				});
 			}
 			else
 			{
@@ -172,8 +153,6 @@ namespace UE::ConcertClientSharedSlate
 			}
 			OnObjectsChangedDelegate.Broadcast(Added, {}, EReplicatedObjectChangeReason::ChangedDirectly);
 		}
-
-		AddObjects(AdditionalObjectsToAdd);
 	}
 
 	void FGenericPropertySelectionModel::RemoveObjects(TConstArrayView<FSoftObjectPath> Objects)
diff --git a/Engine/Plugins/Developer/Concert/ConcertUI/ConcertClientSharedSlate/Source/ConcertClientSharedSlate/Private/Replication/Editor/Model/GenericPropertySelectionModel.h b/Engine/Plugins/Developer/Concert/ConcertUI/ConcertClientSharedSlate/Source/ConcertClientSharedSlate/Private/Replication/Editor/Model/GenericPropertySelectionModel.h
--- a/Engine/Plugins/Developer/Concert/ConcertUI/ConcertClientSharedSlate/Source/ConcertClientSharedSlate/Private/Replication/Editor/Model/GenericPropertySelectionModel.h
+++ b/Engine/Plugins/Developer/Concert/ConcertUI/ConcertClientSharedSlate/Source/ConcertClientSharedSlate/Private/Replication/Editor/Model/GenericPropertySelectionModel.h
@@ -19,10 +19,7 @@ namespace UE::ConcertClientSharedSlate
 	{
 	public:
 		
-		FGenericPropertySelectionModel(
-			TAttribute<FObjectReplicationMap*> ReplicationMapAttribute,
-			TAttribute<const FConcertReplicationEditorSettings*> OptionalReplicationSettingsAttribute = {}
-			);
+		FGenericPropertySelectionModel(TAttribute<FObjectReplicationMap*> ReplicationMapAttribute);
 		
 		//~ Begin IObjectToPropertiesModel Interface
 		virtual uint32 GetNumReplicatedObjects() const override;
@@ -48,9 +45,6 @@ namespace UE::ConcertClientSharedSlate
 		/** Returns the replication map that is supposed to be edited. */
 		TAttribute<FObjectReplicationMap*> ReplicationMapAttribute;
 
-		/** Optional. Contains settings for auto populating objects. */
-		TAttribute<const FConcertReplicationEditorSettings*> OptionalReplicationSettingsAttribute;
-
 		FOnObjectsChanged OnObjectsChangedDelegate;
 		FOnPropertiesChanged OnPropertiesChangedDelegate;
 	};
diff --git a/Engine/Plugins/Developer/Concert/ConcertUI/ConcertClientSharedSlate/Source/ConcertClientSharedSlate/Private/Replication/Editor/Model/TransactionalPropertySelectionModel.cpp b/Engine/Plugins/Developer/Concert/ConcertUI/ConcertClientSharedSlate/Source/ConcertClientSharedSlate/Private/Replication/Editor/Model/TransactionalPropertySelectionModel.cpp
--- a/Engine/Plugins/Developer/Concert/ConcertUI/ConcertClientSharedSlate/Source/ConcertClientSharedSlate/Private/Replication/Editor/Model/TransactionalPropertySelectionModel.cpp
+++ b/Engine/Plugins/Developer/Concert/ConcertUI/ConcertClientSharedSlate/Source/ConcertClientSharedSlate/Private/Replication/Editor/Model/TransactionalPropertySelectionModel.cpp
@@ -10,10 +10,9 @@ namespace UE::ConcertClientSharedSlate
 {
 	FTransactionalPropertySelectionModel::FTransactionalPropertySelectionModel(
 		UObject& OwningObject,
-		TAttribute<FObjectReplicationMap*> ReplicationMapAttribute,
-		TAttribute<const FConcertReplicationEditorSettings*> OptionalReplicationSettingsAttribute
+		TAttribute<FObjectReplicationMap*> ReplicationMapAttribute
 		)
-		: FGenericPropertySelectionModel(MoveTemp(ReplicationMapAttribute), MoveTemp(OptionalReplicationSettingsAttribute))
+		: FGenericPropertySelectionModel(MoveTemp(ReplicationMapAttribute))
 		, OwningObject(&OwningObject)
 	{}
 
diff --git a/Engine/Plugins/Developer/Concert/ConcertUI/ConcertClientSharedSlate/Source/ConcertClientSharedSlate/Private/Replication/Editor/Model/TransactionalPropertySelectionModel.h b/Engine/Plugins/Developer/Concert/ConcertUI/ConcertClientSharedSlate/Source/ConcertClientSharedSlate/Private/Replication/Editor/Model/TransactionalPropertySelectionModel.h
--- a/Engine/Plugins/Developer/Concert/ConcertUI/ConcertClientSharedSlate/Source/ConcertClientSharedSlate/Private/Replication/Editor/Model/TransactionalPropertySelectionModel.h
+++ b/Engine/Plugins/Developer/Concert/ConcertUI/ConcertClientSharedSlate/Source/ConcertClientSharedSlate/Private/Replication/Editor/Model/TransactionalPropertySelectionModel.h
@@ -16,8 +16,7 @@ namespace UE::ConcertClientSharedSlate
 
 		FTransactionalPropertySelectionModel(
 			UObject& OwningObject,
-			TAttribute<FObjectReplicationMap*> ReplicationMapAttribute,
-			TAttribute<const FConcertReplicationEditorSettings*> OptionalReplicationSettingsAttribute = {}
+			TAttribute<FObjectReplicationMap*> ReplicationMapAttribute
 			);
 		
 		//~ Begin IEditableObjectToPropertiesModel Interface
diff --git a/Engine/Plugins/Developer/Concert/ConcertUI/ConcertClientSharedSlate/Source/ConcertClientSharedSlate/Private/Replication/Editor/View/ObjectEditor/SBaseReplicationStreamEditor.cpp b/Engine/Plugins/Developer/Concert/ConcertUI/ConcertClientSharedSlate/Source/ConcertClientSharedSlate/Private/Replication/Editor/View/ObjectEditor/SBaseReplicationStreamEditor.cpp
--- a/Engine/Plugins/Developer/Concert/ConcertUI/ConcertClientSharedSlate/Source/ConcertClientSharedSlate/Private/Replication/Editor/View/ObjectEditor/SBaseReplicationStreamEditor.cpp
+++ b/Engine/Plugins/Developer/Concert/ConcertUI/ConcertClientSharedSlate/Source/ConcertClientSharedSlate/Private/Replication/Editor/View/ObjectEditor/SBaseReplicationStreamEditor.cpp
@@ -10,9 +10,9 @@
 #include "Replication/Editor/Model/ReplicatedObjectData.h"
 #include "Replication/Editor/View/DisplayUtils.h"
 #include "Replication/Editor/View/ObjectViewer/SReplicationStreamViewer.h"
+#include "Replication/Settings/ConcertReplicationEditorSettings.h"
 
 #include "Algo/AnyOf.h"
-#include "GameFramework/Actor.h"
 #include "UObject/Class.h"
 #include "Widgets/SBoxPanel.h"
 
@@ -36,6 +36,7 @@ namespace UE::ConcertClientSharedSlate
 
 		IsEditingEnabledAttribute = InArgs._IsEditingEnabled;
 		EditingDisabledToolTipTextAttribute = InArgs._EditingDisabledToolTipText;
+		ReplicationSettingsAttribute = InArgs._ReplicationSettings;
 		OnExtendObjectsContextMenuDelegate = InArgs._OnExtendObjectsContextMenu;
 		
 		ChildSlot
@@ -69,6 +70,12 @@ namespace UE::ConcertClientSharedSlate
 		];
 	}
 
+	SBaseReplicationStreamEditor::~SBaseReplicationStreamEditor()
+	{
+		EditablePropertiesModel->OnObjectsChanged().RemoveAll(this);
+		EditablePropertiesModel->OnPropertiesChanged().RemoveAll(this);
+	}
+
 	void SBaseReplicationStreamEditor::Refresh()
 	{
 		ReplicationViewer->RefreshObjectData();
@@ -106,6 +113,8 @@ namespace UE::ConcertClientSharedSlate
 		// Newly added objects should be automatically selected
 		if (!AddedObjects.IsEmpty())
 		{
+			AutoAddObjectsAndPropertiesFromSettings(AddedObjects);
+			
 			TArray<FSoftObjectPath> TopLevelObjects;
 			Algo::TransformIf(AddedObjects, TopLevelObjects, [this](const UObject* Object)
 			{
@@ -115,12 +124,45 @@ namespace UE::ConcertClientSharedSlate
 		}
 	}
 
+	void SBaseReplicationStreamEditor::AutoAddObjectsAndPropertiesFromSettings(TConstArrayView<UObject*> AddedObjects)
+	{
+		if (!bIsAddingFromSelection)
+		{
+			return;
+		}
+		
+		const FConcertReplicationEditorSettings* AutoPopulateSettings = ReplicationSettingsAttribute.IsBound()
+		  ? ReplicationSettingsAttribute.Get()
+		  : nullptr;
+		if (!AutoPopulateSettings)
+		{
+			return;
+		}
+
+		for (const UObject* AddedObject : AddedObjects)
+		{
+			TArray<FConcertPropertyChain> AdditionalProperties;
+			AutoPopulateSettings->AddDefaultPropertiesFromSettings(*AddedObject->GetClass(), [&AdditionalProperties](FConcertPropertyChain&& Chain)
+			{
+				AdditionalProperties.Emplace(MoveTemp(Chain));
+			});
+			EditablePropertiesModel->AddProperties({ AddedObject }, AdditionalProperties);
+
+			TArray<UObject*> AdditionalObjectsToAdd;
+			AutoPopulateSettings->AddAdditionalObjectsFromSettings(*AddedObject, [&AdditionalObjectsToAdd](UObject& FurtherObject)
+			{
+				AdditionalObjectsToAdd.Add(&FurtherObject);
+			});
+			EditablePropertiesModel->AddObjects(AdditionalObjectsToAdd);
+		}
+	}
+
 	void SBaseReplicationStreamEditor::OnPropertiesChanged()
 	{
 		ReplicationViewer->RefreshPropertyData();
 	}
 
-	TSharedRef<SWidget> SBaseReplicationStreamEditor::BuildRootAddObjectWidgets() const
+	TSharedRef<SWidget> SBaseReplicationStreamEditor::BuildRootAddObjectWidgets()
 	{
 		using namespace ConcertSharedSlate;
 		
@@ -137,7 +179,7 @@ namespace UE::ConcertClientSharedSlate
 		return Root;
 	}
 
-	void SBaseReplicationStreamEditor::OnObjectsSelectedForAdding(TArray<FSelectableObjectInfo> ObjectsToAdd) const
+	void SBaseReplicationStreamEditor::OnObjectsSelectedForAdding(TArray<FSelectableObjectInfo> ObjectsToAdd)
 	{
 		TArray<UObject*> Objects;
 		Algo::TransformIf(
@@ -146,7 +188,8 @@ namespace UE::ConcertClientSharedSlate
 			[](const FSelectableObjectInfo& SelectableObject) { return SelectableObject.Object.IsValid(); },
 			[](const FSelectableObjectInfo& SelectableObject){ return SelectableObject.Object.Get(); }
 			);
-		
+
+		TGuardValue<bool> GuardObjectSelection(bIsAddingFromSelection, true);
 		EditablePropertiesModel->AddObjects(Objects);
 	}
 
@@ -184,7 +227,7 @@ namespace UE::ConcertClientSharedSlate
 		ReplicationViewer->ClearSubobjectSelection();
 	}
 
-	TSharedPtr<SWidget> SBaseReplicationStreamEditor::OnObjectsContextMenuOpening() const
+	TSharedPtr<SWidget> SBaseReplicationStreamEditor::OnObjectsContextMenuOpening()
 	{
 		FMenuBuilder MenuBuilder(true, nullptr);
 		AddObjectSourceContextMenuOptions(MenuBuilder);
@@ -204,7 +247,7 @@ namespace UE::ConcertClientSharedSlate
 		return MenuBuilder.MakeWidget();
 	}
 
-	void SBaseReplicationStreamEditor::AddObjectSourceContextMenuOptions(FMenuBuilder& MenuBuilder) const
+	void SBaseReplicationStreamEditor::AddObjectSourceContextMenuOptions(FMenuBuilder& MenuBuilder)
 	{
 		using namespace ConcertSharedSlate;
 		
@@ -241,7 +284,7 @@ namespace UE::ConcertClientSharedSlate
 		MenuBuilder.AddSeparator();
 	}
 
-	ConcertSharedSlate::FSourceModelBuilders<FSelectableObjectInfo>::FItemPickerArgs SBaseReplicationStreamEditor::MakeObjectSourceBuilderArgs() const
+	ConcertSharedSlate::FSourceModelBuilders<FSelectableObjectInfo>::FItemPickerArgs SBaseReplicationStreamEditor::MakeObjectSourceBuilderArgs()
 	{
 		using namespace ConcertSharedSlate;
 		using FBuilderDelegates = FSourceModelBuilders<FSelectableObjectInfo>;
diff --git a/Engine/Plugins/Developer/Concert/ConcertUI/ConcertClientSharedSlate/Source/ConcertClientSharedSlate/Private/Replication/Editor/View/ObjectEditor/SBaseReplicationStreamEditor.h b/Engine/Plugins/Developer/Concert/ConcertUI/ConcertClientSharedSlate/Source/ConcertClientSharedSlate/Private/Replication/Editor/View/ObjectEditor/SBaseReplicationStreamEditor.h
--- a/Engine/Plugins/Developer/Concert/ConcertUI/ConcertClientSharedSlate/Source/ConcertClientSharedSlate/Private/Replication/Editor/View/ObjectEditor/SBaseReplicationStreamEditor.h
+++ b/Engine/Plugins/Developer/Concert/ConcertUI/ConcertClientSharedSlate/Source/ConcertClientSharedSlate/Private/Replication/Editor/View/ObjectEditor/SBaseReplicationStreamEditor.h
@@ -11,6 +11,7 @@
 
 class SHorizontalBox;
 struct FConcertPropertyChain;
+struct FConcertReplicationEditorSettings;
 
 namespace UE::ConcertClientSharedSlate
 {
@@ -72,6 +73,9 @@ namespace UE::ConcertClientSharedSlate
 			SLATE_ATTRIBUTE(bool, IsEditingEnabled)
 			/** Optional. Whenever IsEditingEnabled returns true, this tooltip is displayed for relevant, disabled UI. */
 			SLATE_ATTRIBUTE(FText, EditingDisabledToolTipText)
+
+			/** Optional settings for auto adding common properties and objects. */
+			SLATE_ATTRIBUTE(const FConcertReplicationEditorSettings*, ReplicationSettings)
 		SLATE_END_ARGS()
 
 		void Construct(const FArguments& InArgs,
@@ -79,6 +83,7 @@ namespace UE::ConcertClientSharedSlate
            TSharedRef<IObjectSelectionSourceModel> InObjectSelectionSource,
            TSharedRef<IPropertySelectionSourceModel> InPropertySelectionSource
 		);
+		virtual ~SBaseReplicationStreamEditor() override;
 
 		//~ Begin IReplicationStreamEditor Interface
 		virtual void Refresh() override;
@@ -108,27 +113,37 @@ namespace UE::ConcertClientSharedSlate
 		TAttribute<bool> IsEditingEnabledAttribute;
 		/** Optional. Whenever IsEditingEnabled returns true, this tooltip is displayed for relevant, disabled UI. */
 		TAttribute<FText> EditingDisabledToolTipTextAttribute;
+		/** Optional settings for auto adding common properties and objects. */
+		TAttribute<const FConcertReplicationEditorSettings*> ReplicationSettingsAttribute;
 		
 		/** Generates additional entries */
 		FExtendMenu OnExtendObjectsContextMenuDelegate;
 
+		/**
+		 * Set when the combo button is adding an object.
+		 * This is so AutoAddObjectsAndPropertiesFromSettings only adds when an object is added via the combo button.
+		 * Note that AutoAddObjectsAndPropertiesFromSettings must be called in OnObjectsChanged so it becomes part of the transaction.
+		 */
+		bool bIsAddingFromSelection = false;
+
 		bool IsEditingDisabled() const;
 		FText GetEditingDisabledText() const;
 		
 		void OnObjectsChanged(TConstArrayView<UObject*> AddedObjects, TConstArrayView<FSoftObjectPath> RemovedObjects, EReplicatedObjectChangeReason ChangeReason);
+		void AutoAddObjectsAndPropertiesFromSettings(TConstArrayView<UObject*> AddedObjects);
 		void OnPropertiesChanged();
 
 		// Customizing left objects search bar
-		TSharedRef<SWidget> BuildRootAddObjectWidgets() const;
-		void OnObjectsSelectedForAdding(TArray<FSelectableObjectInfo> ObjectsToAdd) const;
+		TSharedRef<SWidget> BuildRootAddObjectWidgets();
+		void OnObjectsSelectedForAdding(TArray<FSelectableObjectInfo> ObjectsToAdd);
 
 		// Respond to object editing events
 		void OnDeleteObjects_PassByValue(TArray<TSharedPtr<FReplicatedObjectData>> CopiedObjectsToDelete) const { OnDeleteObjects(CopiedObjectsToDelete); }
 		void OnDeleteObjects(const TArray<TSharedPtr<FReplicatedObjectData>>& ObjectsToDelete) const;
-		TSharedPtr<SWidget> OnObjectsContextMenuOpening() const;
-		void AddObjectSourceContextMenuOptions(FMenuBuilder& MenuBuilder) const;
+		TSharedPtr<SWidget> OnObjectsContextMenuOpening();
+		void AddObjectSourceContextMenuOptions(FMenuBuilder& MenuBuilder);
 		
 		// Utils for building item source widgets
-		ConcertSharedSlate::FSourceModelBuilders<FSelectableObjectInfo>::FItemPickerArgs MakeObjectSourceBuilderArgs() const;
+		ConcertSharedSlate::FSourceModelBuilders<FSelectableObjectInfo>::FItemPickerArgs MakeObjectSourceBuilderArgs();
 	};
 }
diff --git a/Engine/Plugins/Developer/Concert/ConcertUI/ConcertClientSharedSlate/Source/ConcertClientSharedSlate/Private/Replication/Editor/View/ObjectEditor/SDefaultReplicationStreamEditor.cpp b/Engine/Plugins/Developer/Concert/ConcertUI/ConcertClientSharedSlate/Source/ConcertClientSharedSlate/Private/Replication/Editor/View/ObjectEditor/SDefaultReplicationStreamEditor.cpp
--- a/Engine/Plugins/Developer/Concert/ConcertUI/ConcertClientSharedSlate/Source/ConcertClientSharedSlate/Private/Replication/Editor/View/ObjectEditor/SDefaultReplicationStreamEditor.cpp
+++ b/Engine/Plugins/Developer/Concert/ConcertUI/ConcertClientSharedSlate/Source/ConcertClientSharedSlate/Private/Replication/Editor/View/ObjectEditor/SDefaultReplicationStreamEditor.cpp
@@ -49,6 +49,7 @@ namespace UE::ConcertClientSharedSlate
 				]
 				.IsEditingEnabled(InArgs._IsEditingEnabled)
 				.EditingDisabledToolTipText(InArgs._EditingDisabledToolTipText)
+				.ReplicationSettings(InArgs._ReplicationSettings)
 		];
 	}
 
diff --git a/Engine/Plugins/Developer/Concert/ConcertUI/ConcertClientSharedSlate/Source/ConcertClientSharedSlate/Private/Replication/Editor/View/ObjectEditor/SDefaultReplicationStreamEditor.h b/Engine/Plugins/Developer/Concert/ConcertUI/ConcertClientSharedSlate/Source/ConcertClientSharedSlate/Private/Replication/Editor/View/ObjectEditor/SDefaultReplicationStreamEditor.h
--- a/Engine/Plugins/Developer/Concert/ConcertUI/ConcertClientSharedSlate/Source/ConcertClientSharedSlate/Private/Replication/Editor/View/ObjectEditor/SDefaultReplicationStreamEditor.h
+++ b/Engine/Plugins/Developer/Concert/ConcertUI/ConcertClientSharedSlate/Source/ConcertClientSharedSlate/Private/Replication/Editor/View/ObjectEditor/SDefaultReplicationStreamEditor.h
@@ -10,6 +10,7 @@
 
 class SHorizontalBox;
 struct FConcertPropertyChain;
+struct FConcertReplicationEditorSettings;
 
 namespace UE::ConcertClientSharedSlate
 {
@@ -54,6 +55,9 @@ namespace UE::ConcertClientSharedSlate
 			SLATE_ATTRIBUTE(bool, IsEditingEnabled)
 			/** Optional. Whenever IsEditingEnabled returns true, this tooltip is displayed for relevant, disabled UI. */
 			SLATE_ATTRIBUTE(FText, EditingDisabledToolTipText)
+
+			/** Optional settings for auto adding common properties and objects. */
+			SLATE_ATTRIBUTE(const FConcertReplicationEditorSettings*, ReplicationSettings)
 		SLATE_END_ARGS()
 
 		void Construct(
diff --git a/Engine/Plugins/Developer/Concert/ConcertUI/ConcertClientSharedSlate/Source/ConcertClientSharedSlate/Private/Replication/ReplicationWidgetFactories.cpp b/Engine/Plugins/Developer/Concert/ConcertUI/ConcertClientSharedSlate/Source/ConcertClientSharedSlate/Private/Replication/ReplicationWidgetFactories.cpp
--- a/Engine/Plugins/Developer/Concert/ConcertUI/ConcertClientSharedSlate/Source/ConcertClientSharedSlate/Private/Replication/ReplicationWidgetFactories.cpp
+++ b/Engine/Plugins/Developer/Concert/ConcertUI/ConcertClientSharedSlate/Source/ConcertClientSharedSlate/Private/Replication/ReplicationWidgetFactories.cpp
@@ -27,7 +27,8 @@ namespace UE::ConcertClientSharedSlate
 			.LeftOfObjectSearchBar() [ Params.LeftOfObjectSearchBar.Widget ]
 			.LeftOfPropertySearchBar() [ Params.LeftOfPropertySearchBar.Widget ]
 			.IsEditingEnabled(Params.IsEditingEnabled)
-			.EditingDisabledToolTipText(Params.EditingDisabledToolTipText);
+			.EditingDisabledToolTipText(Params.EditingDisabledToolTipText)
+			.ReplicationSettings(Params.ReplicationSettingsAttribute);
 	}
 
 	TSharedRef<IReplicationStreamEditor> CreateDefaultStreamEditor(FCreateEditorParams Params)
@@ -44,19 +45,18 @@ namespace UE::ConcertClientSharedSlate
 			.LeftOfObjectSearchBar() [ Params.LeftOfObjectSearchBar.Widget ]
 			.LeftOfPropertySearchBar() [ Params.LeftOfPropertySearchBar.Widget ]
 			.IsEditingEnabled(Params.IsEditingEnabled)
-			.EditingDisabledToolTipText(Params.EditingDisabledToolTipText);
+			.EditingDisabledToolTipText(Params.EditingDisabledToolTipText)
+			.ReplicationSettings(Params.ReplicationSettingsAttribute);
 	}
 
 	TSharedRef<IEditableObjectToPropertiesModel> CreatePropertySelectionModel(
 		UObject& OwnerObject,
-		TAttribute<FObjectReplicationMap*> ReplicationMapAttribute,
-		TAttribute<const FConcertReplicationEditorSettings*> OptionalReplicationSettingsAttribute
+		TAttribute<FObjectReplicationMap*> ReplicationMapAttribute
 		)
 	{
 		return MakeShared<FTransactionalPropertySelectionModel>(
 			OwnerObject,
-			MoveTemp(ReplicationMapAttribute),
-			MoveTemp(OptionalReplicationSettingsAttribute)
+			MoveTemp(ReplicationMapAttribute)
 			);
 	}
 }
diff --git a/Engine/Plugins/Developer/Concert/ConcertUI/ConcertClientSharedSlate/Source/ConcertClientSharedSlate/Private/Replication/Settings/ConcertReplicationEditorSettings.cpp b/Engine/Plugins/Developer/Concert/ConcertUI/ConcertClientSharedSlate/Source/ConcertClientSharedSlate/Private/Replication/Settings/ConcertReplicationEditorSettings.cpp
--- a/Engine/Plugins/Developer/Concert/ConcertUI/ConcertClientSharedSlate/Source/ConcertClientSharedSlate/Private/Replication/Settings/ConcertReplicationEditorSettings.cpp
+++ b/Engine/Plugins/Developer/Concert/ConcertUI/ConcertClientSharedSlate/Source/ConcertClientSharedSlate/Private/Replication/Settings/ConcertReplicationEditorSettings.cpp
@@ -14,7 +14,7 @@
 
 namespace UE::ConcertClientSharedSlate::DefaultProperties
 {
-	static void ApplyDefaultPropertySelection(FReplicatedObjectInfo& Info, const FConcertDefaultPropertySelection& Selection, UStruct& Class)
+	static void ApplyDefaultPropertySelection(TFunctionRef<void(FConcertPropertyChain&& Chain)> Callback, const FConcertDefaultPropertySelection& Selection, UStruct& Class)
 	{
 		// Preparse the the FStrings into paths
 		TArray<TArray<FName>> Paths;
@@ -30,13 +30,13 @@ namespace UE::ConcertClientSharedSlate::DefaultProperties
 
 		// This walks through the entire property hierarchy once
 		using namespace ConcertSyncCore::PropertyChain;
-		BulkConstructConcertChainsFromPaths(Class, Paths.Num(), [&Info, &Paths](const FArchiveSerializedPropertyChain& Chain, const FProperty& LeafProperty)
+		BulkConstructConcertChainsFromPaths(Class, Paths.Num(), [&Callback, &Paths](const FArchiveSerializedPropertyChain& Chain, const FProperty& LeafProperty)
 		{
 			const int32 IndexOfMatches = Algo::IndexOfByPredicate(Paths, [&Chain, &LeafProperty](const TArray<FName>& Path){ return DoPathAndChainsMatch(Path, Chain, LeafProperty); });
 			const bool bMatches = IndexOfMatches != INDEX_NONE;
 			if (bMatches)
 			{
-				Info.PropertySelection.ReplicatedProperties.Emplace(&Chain, LeafProperty);
+				Callback(FConcertPropertyChain(&Chain, LeafProperty));
 				Paths.RemoveAtSwap(IndexOfMatches);
 			}
 			return true;
@@ -54,7 +54,7 @@ namespace UE::ConcertClientSharedSlate::DefaultProperties
 	}
 }
 
-void FConcertReplicationEditorSettings::AddDefaultPropertiesFromSettings(FReplicatedObjectInfo& Info, UClass& Class) const
+void FConcertReplicationEditorSettings::AddDefaultPropertiesFromSettings(UClass& Class, TFunctionRef<void(FConcertPropertyChain&& Chain)> Callback) const
 {
 	// Find the most specialized class properties
 	UClass* Current = &Class;
@@ -67,12 +67,12 @@ void FConcertReplicationEditorSettings::AddDefaultPropertiesFromSettings(FReplic
 			continue;
 		}
 			
-		UE::ConcertClientSharedSlate::DefaultProperties::ApplyDefaultPropertySelection(Info, *DefaultProperties, *Current);
+		UE::ConcertClientSharedSlate::DefaultProperties::ApplyDefaultPropertySelection(Callback, *DefaultProperties, *Current);
 		// Recurse super structs
 		if (UClass* Parent = Current->GetSuperClass()
 			; Parent && DefaultProperties->bInheritFromBase)
 		{
-			AddDefaultPropertiesFromSettings(Info, *Parent);
+			AddDefaultPropertiesFromSettings(*Parent, Callback);
 		}
 	}
 }
@@ -94,7 +94,7 @@ namespace UE::ConcertClientSharedSlate::DefaultSubobjects
 		return false;
 	}
 	
-	static void ApplyDefaultSubobjectSelection(UObject& AddedObject, const FConcertDefaultSubobjectSelection& Selection, TFunctionRef<void(UObject&)> FurtherObjectsCallback)
+	static void ApplyDefaultSubobjectSelection(const UObject& AddedObject, const FConcertDefaultSubobjectSelection& Selection, TFunctionRef<void(UObject&)> FurtherObjectsCallback)
 	{
 		// Do not search recursively so FurtherObjectsCallback can decide to call AddAdditionalObjectsFromSettings again on the newly added objects.
 		constexpr bool bIncludeNested = false;
@@ -137,7 +137,7 @@ namespace UE::ConcertClientSharedSlate::DefaultSubobjects
 		}, bIncludeNested);
 	}
 
-	static void InternalAddAdditionalObjectsFromSettings(UClass& StartClass, const FConcertReplicationEditorSettings& Settings, UObject& AddedObject, TFunctionRef<void(UObject&)> FurtherObjectsCallback)
+	static void InternalAddAdditionalObjectsFromSettings(UClass& StartClass, const FConcertReplicationEditorSettings& Settings, const UObject& AddedObject, TFunctionRef<void(UObject&)> FurtherObjectsCallback)
 	{
 		// Find the most specialized class properties
 		UClass* Current = &StartClass;
@@ -161,7 +161,7 @@ namespace UE::ConcertClientSharedSlate::DefaultSubobjects
 	}
 }
 
-void FConcertReplicationEditorSettings::AddAdditionalObjectsFromSettings(UObject& AddedObject, TFunctionRef<void(UObject&)> FurtherObjectsCallback) const
+void FConcertReplicationEditorSettings::AddAdditionalObjectsFromSettings(const UObject& AddedObject, TFunctionRef<void(UObject&)> FurtherObjectsCallback) const
 {
 	UE::ConcertClientSharedSlate::DefaultSubobjects::InternalAddAdditionalObjectsFromSettings(*AddedObject.GetClass(), *this, AddedObject, FurtherObjectsCallback);
 }
\ No newline at end of file
diff --git a/Engine/Plugins/Developer/Concert/ConcertUI/ConcertClientSharedSlate/Source/ConcertClientSharedSlate/Public/Replication/ReplicationWidgetFactories.h b/Engine/Plugins/Developer/Concert/ConcertUI/ConcertClientSharedSlate/Source/ConcertClientSharedSlate/Public/Replication/ReplicationWidgetFactories.h
--- a/Engine/Plugins/Developer/Concert/ConcertUI/ConcertClientSharedSlate/Source/ConcertClientSharedSlate/Public/Replication/ReplicationWidgetFactories.h
+++ b/Engine/Plugins/Developer/Concert/ConcertUI/ConcertClientSharedSlate/Source/ConcertClientSharedSlate/Public/Replication/ReplicationWidgetFactories.h
@@ -37,8 +37,7 @@ namespace UE::ConcertClientSharedSlate
 	 */
 	CONCERTCLIENTSHAREDSLATE_API TSharedRef<IEditableObjectToPropertiesModel> CreatePropertySelectionModel(
 		UObject& OwnerObject,
-		TAttribute<FObjectReplicationMap*> ReplicationMapAttribute,
-		TAttribute<const FConcertReplicationEditorSettings*> OptionalReplicationSettingsAttribute = {}
+		TAttribute<FObjectReplicationMap*> ReplicationMapAttribute
 		);
 	
 	struct FCreateSubobjectViewParams
@@ -98,6 +97,9 @@ namespace UE::ConcertClientSharedSlate
 		TAttribute<bool> IsEditingEnabled;
 		/** Optional. Whenever IsEditingEnabled returns true, this tooltip is displayed for relevant, disabled UI. */
 		TAttribute<FText> EditingDisabledToolTipText;
+
+		/** Optional settings for auto adding common properties and objects. */
+		TAttribute<const FConcertReplicationEditorSettings*> ReplicationSettingsAttribute;
 	};
 
 	/**
diff --git a/Engine/Plugins/Developer/Concert/ConcertUI/ConcertClientSharedSlate/Source/ConcertClientSharedSlate/Public/Replication/Settings/ConcertReplicationEditorSettings.h b/Engine/Plugins/Developer/Concert/ConcertUI/ConcertClientSharedSlate/Source/ConcertClientSharedSlate/Public/Replication/Settings/ConcertReplicationEditorSettings.h
--- a/Engine/Plugins/Developer/Concert/ConcertUI/ConcertClientSharedSlate/Source/ConcertClientSharedSlate/Public/Replication/Settings/ConcertReplicationEditorSettings.h
+++ b/Engine/Plugins/Developer/Concert/ConcertUI/ConcertClientSharedSlate/Source/ConcertClientSharedSlate/Public/Replication/Settings/ConcertReplicationEditorSettings.h
@@ -4,8 +4,11 @@
 
 #include "ConcertDefaultPropertySelection.h"
 #include "ConcertDefaultSubobjectSelection.h"
+#include "Templates/Function.h"
 #include "ConcertReplicationEditorSettings.generated.h"
 
+class UObject;
+struct FConcertPropertyChain;
 struct FReplicatedObjectInfo;
 
 USTRUCT()
@@ -26,9 +29,9 @@ struct CONCERTCLIENTSHAREDSLATE_API FConcertReplicationEditorSettings
 	UPROPERTY(EditAnywhere, Config, Category = "Replication|Editor")
 	TMap<FSoftClassPath, FConcertDefaultSubobjectSelection> DefaultSubobjectSelection;
 
-	/** Reads DefaultPropertySelection and applies any default property selections to Info based on the Class just added. */
-	void AddDefaultPropertiesFromSettings(FReplicatedObjectInfo& Info, UClass& Class) const;
+	/** Reads DefaultPropertySelection and calls Callback for any default property selections based on the Class just added. */
+	void AddDefaultPropertiesFromSettings(UClass& Class, TFunctionRef<void(FConcertPropertyChain&& Chain)> Callback) const;
 
 	/** Reads DefaultComponentSelection and calls FurtherObjectsCallback on any further objects that should also be added. */
-	void AddAdditionalObjectsFromSettings(UObject& AddedObject, TFunctionRef<void(UObject&)> FurtherObjectsCallback) const;
+	void AddAdditionalObjectsFromSettings(const UObject& AddedObject, TFunctionRef<void(UObject&)> FurtherObjectsCallback) const;
 };
\ No newline at end of file
