Rename IObjectToPropertiesModel to IReplicationStreamModel

diff --git a/Engine/Plugins/Developer/Concert/ConcertApp/MultiUserClient/Source/MultiUserClient/Private/Replication/Client/ReplicationClient.cpp b/Engine/Plugins/Developer/Concert/ConcertApp/MultiUserClient/Source/MultiUserClient/Private/Replication/Client/ReplicationClient.cpp
--- a/Engine/Plugins/Developer/Concert/ConcertApp/MultiUserClient/Source/MultiUserClient/Private/Replication/Client/ReplicationClient.cpp
+++ b/Engine/Plugins/Developer/Concert/ConcertApp/MultiUserClient/Source/MultiUserClient/Private/Replication/Client/ReplicationClient.cpp
@@ -3,7 +3,7 @@
 #include "ReplicationClient.h"
 
 #include "Assets/MultiUserReplicationClientPreset.h"
-#include "Replication/Editor/Model/IEditableObjectToPropertiesModel.h"
+#include "Replication/Editor/Model/IEditableReplicationStreamModel.h"
 #include "Replication/ReplicationWidgetFactories.h"
 #include "Replication/Stream/StreamChangeTracker.h"
 #include "Replication/Submission/Data/AuthoritySubmission.h"
diff --git a/Engine/Plugins/Developer/Concert/ConcertApp/MultiUserClient/Source/MultiUserClient/Private/Replication/Client/ReplicationClient.h b/Engine/Plugins/Developer/Concert/ConcertApp/MultiUserClient/Source/MultiUserClient/Private/Replication/Client/ReplicationClient.h
--- a/Engine/Plugins/Developer/Concert/ConcertApp/MultiUserClient/Source/MultiUserClient/Private/Replication/Client/ReplicationClient.h
+++ b/Engine/Plugins/Developer/Concert/ConcertApp/MultiUserClient/Source/MultiUserClient/Private/Replication/Client/ReplicationClient.h
@@ -16,7 +16,7 @@ class UMultiUserReplicationClientPreset;
 
 namespace UE::ConcertClientSharedSlate
 {
-	class IEditableObjectToPropertiesModel;
+	class IEditableReplicationStreamModel;
 	enum class EReplicatedObjectChangeReason : uint8;
 }
 
@@ -45,7 +45,7 @@ namespace UE::MultiUserClient
 			);
 
 		UMultiUserReplicationClientPreset* GetClientContent() const { return ClientContentStorage; }
-		TSharedRef<ConcertClientSharedSlate::IEditableObjectToPropertiesModel> GetClientEditModel() const { return LocalClientEditModel; }
+		TSharedRef<ConcertClientSharedSlate::IEditableReplicationStreamModel> GetClientEditModel() const { return LocalClientEditModel; }
 		IClientStreamSynchronizer& GetStreamSynchronizer() const { return *StreamSynchronizer.Get(); }
 		IClientAuthoritySynchronizer& GetAuthoritySynchronizer() const { return *AuthoritySynchronizer.Get(); }
 		
@@ -62,7 +62,7 @@ namespace UE::MultiUserClient
 
 		/**
 		 * Called when the data underlying the model has changed externally. Since the change was not caused by the model,
-		 * its events, like IEditableObjectToPropertiesModel::OnObjectsChanged, were not called.
+		 * its events, like IEditableReplicationStreamModel::OnObjectsChanged, were not called.
 		 * 
 		 * Subscribers are intended to call IReplicationStreamEditor::Refresh() in response.
 		 * 
@@ -94,7 +94,7 @@ namespace UE::MultiUserClient
 		 * The UI is destroyed right after this FClientStreamRepository.
 		 * @see FMultiUserReplicationManager::OnLeaveSession
 		 */
-		TSharedRef<ConcertClientSharedSlate::IEditableObjectToPropertiesModel> LocalClientEditModel;
+		TSharedRef<ConcertClientSharedSlate::IEditableReplicationStreamModel> LocalClientEditModel;
 		
 		/** Tracks changes made to server's state of the client's streams and prepares to upload them using StreamSynchronizer. */
 		FStreamChangeTracker LocalClientStreamDiffer;
diff --git a/Engine/Plugins/Developer/Concert/ConcertApp/MultiUserClient/Source/MultiUserClient/Private/Replication/Client/ReplicationClientManager.h b/Engine/Plugins/Developer/Concert/ConcertApp/MultiUserClient/Source/MultiUserClient/Private/Replication/Client/ReplicationClientManager.h
--- a/Engine/Plugins/Developer/Concert/ConcertApp/MultiUserClient/Source/MultiUserClient/Private/Replication/Client/ReplicationClientManager.h
+++ b/Engine/Plugins/Developer/Concert/ConcertApp/MultiUserClient/Source/MultiUserClient/Private/Replication/Client/ReplicationClientManager.h
@@ -23,7 +23,7 @@ struct FConcertSessionClientInfo;
 
 namespace UE::ConcertClientSharedSlate
 {
-	class IEditableObjectToPropertiesModel;
+	class IEditableReplicationStreamModel;
 }
 
 namespace UE::MultiUserClient
diff --git a/Engine/Plugins/Developer/Concert/ConcertApp/MultiUserClient/Source/MultiUserClient/Private/Replication/MultiUserReplicationManager.cpp b/Engine/Plugins/Developer/Concert/ConcertApp/MultiUserClient/Source/MultiUserClient/Private/Replication/MultiUserReplicationManager.cpp
--- a/Engine/Plugins/Developer/Concert/ConcertApp/MultiUserClient/Source/MultiUserClient/Private/Replication/MultiUserReplicationManager.cpp
+++ b/Engine/Plugins/Developer/Concert/ConcertApp/MultiUserClient/Source/MultiUserClient/Private/Replication/MultiUserReplicationManager.cpp
@@ -73,9 +73,9 @@ namespace UE::MultiUserClient
 
 	void FMultiUserReplicationManager::OnLeaveSession(IConcertClientSession&)
 	{
-		// This clears the UI. The clients' IEditableObjectToPropertiesModels should no longer be referenced by anyone.
+		// This clears the UI. The clients' IEditableReplicationStreamModels should no longer be referenced by anyone.
 		SetConnectionStateAndBroadcast(EMultiUserReplicationConnectionState::Disconnected);
-		// Keep in mind the IEditableObjectToPropertiesModels were referenced by the UI so call this after clearing the UI.
+		// Keep in mind the IEditableReplicationStreamModels were referenced by the UI so call this after clearing the UI.
 		ConnectedState.Reset();
 	}
 
diff --git a/Engine/Plugins/Developer/Concert/ConcertApp/MultiUserClient/Source/MultiUserClient/Private/Replication/MultiUserReplicationManager.h b/Engine/Plugins/Developer/Concert/ConcertApp/MultiUserClient/Source/MultiUserClient/Private/Replication/MultiUserReplicationManager.h
--- a/Engine/Plugins/Developer/Concert/ConcertApp/MultiUserClient/Source/MultiUserClient/Private/Replication/MultiUserReplicationManager.h
+++ b/Engine/Plugins/Developer/Concert/ConcertApp/MultiUserClient/Source/MultiUserClient/Private/Replication/MultiUserReplicationManager.h
@@ -20,7 +20,7 @@ namespace UE::ConcertSyncClient::Replication
 
 namespace UE::ConcertClientSharedSlate
 {
-	class IEditableObjectToPropertiesModel;
+	class IEditableReplicationStreamModel;
 }
 
 namespace UE::MultiUserClient
diff --git a/Engine/Plugins/Developer/Concert/ConcertApp/MultiUserClient/Source/MultiUserClient/Private/Replication/Submission/AutoSubmissionPolicy.cpp b/Engine/Plugins/Developer/Concert/ConcertApp/MultiUserClient/Source/MultiUserClient/Private/Replication/Submission/AutoSubmissionPolicy.cpp
--- a/Engine/Plugins/Developer/Concert/ConcertApp/MultiUserClient/Source/MultiUserClient/Private/Replication/Submission/AutoSubmissionPolicy.cpp
+++ b/Engine/Plugins/Developer/Concert/ConcertApp/MultiUserClient/Source/MultiUserClient/Private/Replication/Submission/AutoSubmissionPolicy.cpp
@@ -3,7 +3,7 @@
 #include "AutoSubmissionPolicy.h"
 
 #include "ISubmissionWorkflow.h"
-#include "Replication/Editor/Model/IEditableObjectToPropertiesModel.h"
+#include "Replication/Editor/Model/IEditableReplicationStreamModel.h"
 
 #include "Misc/CoreDelegates.h"
 #include "Replication/Authority/AuthorityChangeTracker.h"
@@ -12,7 +12,7 @@ namespace UE::MultiUserClient
 {
 	FAutoSubmissionPolicy::FAutoSubmissionPolicy(
 		ISubmissionWorkflow& InSubmissionWorkflow,
-		ConcertClientSharedSlate::IEditableObjectToPropertiesModel& InStreamEditorModel,
+		ConcertClientSharedSlate::IEditableReplicationStreamModel& InStreamEditorModel,
 		FAuthorityChangeTracker& InAuthorityChangeTracker
 		)
 		: SubmissionWorkflow(InSubmissionWorkflow)
diff --git a/Engine/Plugins/Developer/Concert/ConcertApp/MultiUserClient/Source/MultiUserClient/Private/Replication/Submission/AutoSubmissionPolicy.h b/Engine/Plugins/Developer/Concert/ConcertApp/MultiUserClient/Source/MultiUserClient/Private/Replication/Submission/AutoSubmissionPolicy.h
--- a/Engine/Plugins/Developer/Concert/ConcertApp/MultiUserClient/Source/MultiUserClient/Private/Replication/Submission/AutoSubmissionPolicy.h
+++ b/Engine/Plugins/Developer/Concert/ConcertApp/MultiUserClient/Source/MultiUserClient/Private/Replication/Submission/AutoSubmissionPolicy.h
@@ -2,12 +2,12 @@
 
 #pragma once
 
-#include "Replication/Editor/Model/IEditableObjectToPropertiesModel.h"
+#include "Replication/Editor/Model/IEditableReplicationStreamModel.h"
 #include "Templates/UnrealTemplate.h"
 
 namespace UE::ConcertClientSharedSlate
 {
-	class IEditableObjectToPropertiesModel;
+	class IEditableReplicationStreamModel;
 }
 
 namespace UE::MultiUserClient
@@ -28,7 +28,7 @@ namespace UE::MultiUserClient
 		
 		FAutoSubmissionPolicy(
 			ISubmissionWorkflow& InSubmissionWorkflow,
-			ConcertClientSharedSlate::IEditableObjectToPropertiesModel& InStreamEditorModel,
+			ConcertClientSharedSlate::IEditableReplicationStreamModel& InStreamEditorModel,
 			FAuthorityChangeTracker& InAuthorityChangeTracker
 			);
 		~FAutoSubmissionPolicy();
@@ -39,7 +39,7 @@ namespace UE::MultiUserClient
 		ISubmissionWorkflow& SubmissionWorkflow;
 
 		/** Informs us when the stream is structurally changed by the user. */
-		ConcertClientSharedSlate::IEditableObjectToPropertiesModel& StreamEditorModel;
+		ConcertClientSharedSlate::IEditableReplicationStreamModel& StreamEditorModel;
 		/** Informs us when authority is changed by the user. */
 		FAuthorityChangeTracker& AuthorityChangeTracker;
 		
diff --git a/Engine/Plugins/Developer/Concert/ConcertApp/MultiUserClient/Source/MultiUserClient/Private/Widgets/ActiveSession/Replication/Client/Columns/SingleClientColumns.h b/Engine/Plugins/Developer/Concert/ConcertApp/MultiUserClient/Source/MultiUserClient/Private/Widgets/ActiveSession/Replication/Client/Columns/SingleClientColumns.h
--- a/Engine/Plugins/Developer/Concert/ConcertApp/MultiUserClient/Source/MultiUserClient/Private/Widgets/ActiveSession/Replication/Client/Columns/SingleClientColumns.h
+++ b/Engine/Plugins/Developer/Concert/ConcertApp/MultiUserClient/Source/MultiUserClient/Private/Widgets/ActiveSession/Replication/Client/Columns/SingleClientColumns.h
@@ -11,7 +11,7 @@ class IConcertClient;
 namespace UE::ConcertClientSharedSlate
 {
 	class IReplicationStreamViewer;
-	class IObjectToPropertiesModel;
+	class IReplicationStreamModel;
 }
 
 namespace UE::MultiUserClient
@@ -59,7 +59,7 @@ namespace UE::MultiUserClient::SingleClientColumns
 	 * @return Column that can be placed in the table
 	 */
 	ConcertClientSharedSlate::ReplicationColumns::FReplicationTopLevelObjectColumn ToggleTopLevelAuthority(
-		ConcertClientSharedSlate::IObjectToPropertiesModel& ClientStreamModel,
+		ConcertClientSharedSlate::IReplicationStreamModel& ClientStreamModel,
 		FAuthorityChangeTracker& ChangeTracker,
 		ISubmissionWorkflow& SubmissionWorkflow
 		);
@@ -93,7 +93,7 @@ namespace UE::MultiUserClient::SingleClientColumns
 	ConcertClientSharedSlate::ReplicationColumns::FReplicationTopLevelObjectColumn OwnerOfTopLevelObject(
 		const TSharedRef<IConcertClient>& InClient,
 		FGlobalAuthorityCache& InAuthorityCache,
-		ConcertClientSharedSlate::IObjectToPropertiesModel& InObjectModel
+		ConcertClientSharedSlate::IReplicationStreamModel& InObjectModel
 		);
 
 	/**
diff --git a/Engine/Plugins/Developer/Concert/ConcertApp/MultiUserClient/Source/MultiUserClient/Private/Widgets/ActiveSession/Replication/Client/Columns/SingleClientColumns_Owner.cpp b/Engine/Plugins/Developer/Concert/ConcertApp/MultiUserClient/Source/MultiUserClient/Private/Widgets/ActiveSession/Replication/Client/Columns/SingleClientColumns_Owner.cpp
--- a/Engine/Plugins/Developer/Concert/ConcertApp/MultiUserClient/Source/MultiUserClient/Private/Widgets/ActiveSession/Replication/Client/Columns/SingleClientColumns_Owner.cpp
+++ b/Engine/Plugins/Developer/Concert/ConcertApp/MultiUserClient/Source/MultiUserClient/Private/Widgets/ActiveSession/Replication/Client/Columns/SingleClientColumns_Owner.cpp
@@ -34,7 +34,7 @@ namespace UE::MultiUserClient::SingleClientColumns
 	{
 		static TArray<FGuid> GetOwnersOfTopLevelObject(
 			const FGlobalAuthorityCache& InAuthorityCache,
-			ConcertClientSharedSlate::IObjectToPropertiesModel& InObjectModel,
+			ConcertClientSharedSlate::IReplicationStreamModel& InObjectModel,
 			const FSoftObjectPath& InTopLevelObject
 			)
 		{
@@ -54,7 +54,7 @@ namespace UE::MultiUserClient::SingleClientColumns
 	ConcertClientSharedSlate::ReplicationColumns::FReplicationTopLevelObjectColumn OwnerOfTopLevelObject(
 		const TSharedRef<IConcertClient>& InClient,
 		FGlobalAuthorityCache& InAuthorityCache,
-		ConcertClientSharedSlate::IObjectToPropertiesModel& InObjectModel
+		ConcertClientSharedSlate::IReplicationStreamModel& InObjectModel
 		)
 	{
 		using FColumnType = ConcertClientSharedSlate::ReplicationColumns::FReplicationSubobjectObjectColumn;
diff --git a/Engine/Plugins/Developer/Concert/ConcertApp/MultiUserClient/Source/MultiUserClient/Private/Widgets/ActiveSession/Replication/Client/Columns/SingleClientColumns_ToggleAuthority.cpp b/Engine/Plugins/Developer/Concert/ConcertApp/MultiUserClient/Source/MultiUserClient/Private/Widgets/ActiveSession/Replication/Client/Columns/SingleClientColumns_ToggleAuthority.cpp
--- a/Engine/Plugins/Developer/Concert/ConcertApp/MultiUserClient/Source/MultiUserClient/Private/Widgets/ActiveSession/Replication/Client/Columns/SingleClientColumns_ToggleAuthority.cpp
+++ b/Engine/Plugins/Developer/Concert/ConcertApp/MultiUserClient/Source/MultiUserClient/Private/Widgets/ActiveSession/Replication/Client/Columns/SingleClientColumns_ToggleAuthority.cpp
@@ -7,7 +7,7 @@
 #include "Replication/Authority/AuthorityChangeTracker.h"
 #include "Replication/Authority/EAuthorityMutability.h"
 #include "Replication/Authority/IClientAuthoritySynchronizer.h"
-#include "Replication/Editor/Model/IObjectToPropertiesModel.h"
+#include "Replication/Editor/Model/IReplicationStreamModel.h"
 #include "Replication/Editor/View/IReplicationStreamViewer.h"
 #include "Replication/Editor/View/ReplicationColumnsUtils.h"
 #include "Replication/Submission/ISubmissionWorkflow.h"
@@ -25,7 +25,7 @@ namespace UE::MultiUserClient::SingleClientColumns
 	{
 		static ECheckBoxState GetCheckboxState(
 			const ConcertClientSharedSlate::FReplicatedObjectData& ObjectData,
-			ConcertClientSharedSlate::IObjectToPropertiesModel* ClientStreamModel,
+			ConcertClientSharedSlate::IReplicationStreamModel* ClientStreamModel,
 			FAuthorityChangeTracker* ChangeTracker
 			)
 		{
@@ -76,7 +76,7 @@ namespace UE::MultiUserClient::SingleClientColumns
 		static void OnCheckboxStateChanged(
 			bool bIsChecked,
 			const ConcertClientSharedSlate::FReplicatedObjectData& ObjectData,
-			ConcertClientSharedSlate::IObjectToPropertiesModel* ClientStreamModel,
+			ConcertClientSharedSlate::IReplicationStreamModel* ClientStreamModel,
 			FAuthorityChangeTracker* ChangeTracker
 			)
 		{
@@ -93,7 +93,7 @@ namespace UE::MultiUserClient::SingleClientColumns
 
 		static bool IsEnabled(
 			const ConcertClientSharedSlate::FReplicatedObjectData& ObjectData,
-			ConcertClientSharedSlate::IObjectToPropertiesModel* ClientStreamModel,
+			ConcertClientSharedSlate::IReplicationStreamModel* ClientStreamModel,
 			FAuthorityChangeTracker* ChangeTracker
 			)
 		{
@@ -115,7 +115,7 @@ namespace UE::MultiUserClient::SingleClientColumns
 	}
 	
 	ConcertClientSharedSlate::ReplicationColumns::FReplicationTopLevelObjectColumn ToggleTopLevelAuthority(
-		ConcertClientSharedSlate::IObjectToPropertiesModel& ClientStreamModel,
+		ConcertClientSharedSlate::IReplicationStreamModel& ClientStreamModel,
 		FAuthorityChangeTracker& ChangeTracker,
 		ISubmissionWorkflow& SubmissionWorkflow
 		)
diff --git a/Engine/Plugins/Developer/Concert/ConcertApp/MultiUserClient/Source/MultiUserClient/Private/Widgets/ActiveSession/Replication/Client/SClientToolbar.cpp b/Engine/Plugins/Developer/Concert/ConcertApp/MultiUserClient/Source/MultiUserClient/Private/Widgets/ActiveSession/Replication/Client/SClientToolbar.cpp
new file mode 100644
--- /dev/null
+++ b/Engine/Plugins/Developer/Concert/ConcertApp/MultiUserClient/Source/MultiUserClient/Private/Widgets/ActiveSession/Replication/Client/SClientToolbar.cpp
@@ -0,0 +1,41 @@
+// Copyright Epic Games, Inc. All Rights Reserved.
+
+#include "SClientToolbar.h"
+
+#include "Widgets/ActiveSession/Replication/Client/SReplicationStatus.h"
+
+#include "Widgets/Layout/SSpacer.h"
+#include "Widgets/SBoxPanel.h"
+
+namespace UE::MultiUserClient
+{
+	void SClientToolbar::Construct(const FArguments& InArgs, const ConcertClientSharedSlate::IReplicationStreamModel& InObjectModel, FGlobalAuthorityCache& InAuthorityCache)
+	{
+		ChildSlot
+		[
+			SNew(SHorizontalBox)
+
+			+SHorizontalBox::Slot()
+			.AutoWidth()
+			.Padding(5.f, 0.f)
+			[
+				InArgs._ViewSelectionArea.Widget
+			]
+			
+			+SHorizontalBox::Slot()
+			.FillWidth(1.f)
+			[
+				SNew(SSpacer)
+			]
+
+			+SHorizontalBox::Slot()
+			.AutoWidth()
+			.VAlign(VAlign_Center)
+			.Padding(0.f, 0.f, 5.f, 0.f)
+			[
+				SNew(SReplicationStatus, InObjectModel, InAuthorityCache)
+				.DisplayedClients(InArgs._DisplayedClients)
+			]
+		];
+	}
+}
\ No newline at end of file
diff --git a/Engine/Plugins/Developer/Concert/ConcertApp/MultiUserClient/Source/MultiUserClient/Private/Widgets/ActiveSession/Replication/Client/SClientToolbar.h b/Engine/Plugins/Developer/Concert/ConcertApp/MultiUserClient/Source/MultiUserClient/Private/Widgets/ActiveSession/Replication/Client/SClientToolbar.h
new file mode 100644
--- /dev/null
+++ b/Engine/Plugins/Developer/Concert/ConcertApp/MultiUserClient/Source/MultiUserClient/Private/Widgets/ActiveSession/Replication/Client/SClientToolbar.h
@@ -0,0 +1,38 @@
+// Copyright Epic Games, Inc. All Rights Reserved.
+
+#pragma once
+
+#include "Widgets/DeclarativeSyntaxSupport.h"
+#include "Widgets/SCompoundWidget.h"
+
+namespace UE::ConcertClientSharedSlate
+{
+	class IReplicationStreamModel;
+}
+
+namespace UE::MultiUserClient
+{
+	class FGlobalAuthorityCache;
+	class FStreamChangeTracker;
+	class ISubmissionWorkflow;
+
+	/**
+	 * Contains a bunch of actions that can be performed on client view.
+	 */
+	class SClientToolbar : public SCompoundWidget
+	{
+	public:
+		
+		SLATE_BEGIN_ARGS(SClientToolbar)
+		{}
+			/** Dedicated space for a widget with which to change the view. */
+			SLATE_NAMED_SLOT(FArguments, ViewSelectionArea)
+		
+			/** The clients to show statistics for */
+			SLATE_ATTRIBUTE(TSet<FGuid>, DisplayedClients)
+		SLATE_END_ARGS()
+
+		void Construct(const FArguments& InArgs, const ConcertClientSharedSlate::IReplicationStreamModel& InObjectModel, FGlobalAuthorityCache& InAuthorityCache);
+	};
+}
+
diff --git a/Engine/Plugins/Developer/Concert/ConcertApp/MultiUserClient/Source/MultiUserClient/Private/Widgets/ActiveSession/Replication/Client/SReplicationStatus.cpp b/Engine/Plugins/Developer/Concert/ConcertApp/MultiUserClient/Source/MultiUserClient/Private/Widgets/ActiveSession/Replication/Client/SReplicationStatus.cpp
--- a/Engine/Plugins/Developer/Concert/ConcertApp/MultiUserClient/Source/MultiUserClient/Private/Widgets/ActiveSession/Replication/Client/SReplicationStatus.cpp
+++ b/Engine/Plugins/Developer/Concert/ConcertApp/MultiUserClient/Source/MultiUserClient/Private/Widgets/ActiveSession/Replication/Client/SReplicationStatus.cpp
@@ -45,7 +45,7 @@ namespace UE::MultiUserClient
 		}
 	}
 	
-	void SReplicationStatus::Construct(const FArguments& InArgs, const ConcertClientSharedSlate::IObjectToPropertiesModel& InObjectModel, FGlobalAuthorityCache& InAuthorityCache)
+	void SReplicationStatus::Construct(const FArguments& InArgs, const ConcertClientSharedSlate::IReplicationStreamModel& InObjectModel, FGlobalAuthorityCache& InAuthorityCache)
 	{
 		ObjectModel = &InObjectModel;
 		AuthorityCache = &InAuthorityCache;
diff --git a/Engine/Plugins/Developer/Concert/ConcertApp/MultiUserClient/Source/MultiUserClient/Private/Widgets/ActiveSession/Replication/Client/SReplicationStatus.h b/Engine/Plugins/Developer/Concert/ConcertApp/MultiUserClient/Source/MultiUserClient/Private/Widgets/ActiveSession/Replication/Client/SReplicationStatus.h
--- a/Engine/Plugins/Developer/Concert/ConcertApp/MultiUserClient/Source/MultiUserClient/Private/Widgets/ActiveSession/Replication/Client/SReplicationStatus.h
+++ b/Engine/Plugins/Developer/Concert/ConcertApp/MultiUserClient/Source/MultiUserClient/Private/Widgets/ActiveSession/Replication/Client/SReplicationStatus.h
@@ -9,7 +9,7 @@ class STextBlock;
 
 namespace UE::ConcertClientSharedSlate
 {
-	class IObjectToPropertiesModel;
+	class IReplicationStreamModel;
 }
 
 namespace UE::MultiUserClient
@@ -33,13 +33,13 @@ namespace UE::MultiUserClient
 			SLATE_ATTRIBUTE(TSet<FGuid>, DisplayedClients)
 		SLATE_END_ARGS()
 
-		void Construct(const FArguments& InArgs, const ConcertClientSharedSlate::IObjectToPropertiesModel& InObjectModel, FGlobalAuthorityCache& InAuthorityCache);
+		void Construct(const FArguments& InArgs, const ConcertClientSharedSlate::IReplicationStreamModel& InObjectModel, FGlobalAuthorityCache& InAuthorityCache);
 		virtual ~SReplicationStatus() override;
 
 	private:
 
 		/** Used to look up registered objects and subobjects. */
-		const ConcertClientSharedSlate::IObjectToPropertiesModel* ObjectModel;
+		const ConcertClientSharedSlate::IReplicationStreamModel* ObjectModel;
 		/** Used to get authority state of objects and informs us when authority changes. */
 		FGlobalAuthorityCache* AuthorityCache;
 
diff --git a/Engine/Plugins/Developer/Concert/ConcertApp/MultiUserClient/Source/MultiUserClient/Private/Widgets/ActiveSession/Replication/Client/Single/SReplicationClientView.cpp b/Engine/Plugins/Developer/Concert/ConcertApp/MultiUserClient/Source/MultiUserClient/Private/Widgets/ActiveSession/Replication/Client/Single/SReplicationClientView.cpp
--- a/Engine/Plugins/Developer/Concert/ConcertApp/MultiUserClient/Source/MultiUserClient/Private/Widgets/ActiveSession/Replication/Client/Single/SReplicationClientView.cpp
+++ b/Engine/Plugins/Developer/Concert/ConcertApp/MultiUserClient/Source/MultiUserClient/Private/Widgets/ActiveSession/Replication/Client/Single/SReplicationClientView.cpp
@@ -4,7 +4,7 @@
 
 #include "MultiUserReplicationSettings.h"
 #include "Replication/Client/ReplicationClient.h"
-#include "Replication/Editor/Model/IEditableObjectToPropertiesModel.h"
+#include "Replication/Editor/Model/IReplicationStreamModel.h"
 #include "Replication/Editor/Model/Object/ActorSelectionSourceModel.h"
 #include "Replication/Editor/Model/Property/SelectPropertyFromUClassModel.h"
 #include "Replication/Editor/View/IReplicationStreamEditor.h"
@@ -12,7 +12,7 @@
 #include "Replication/Client/ReplicationClientManager.h"
 #include "Replication/Submission/ISubmissionWorkflow.h"
 #include "Widgets/ActiveSession/Replication/Client/Columns/SingleClientColumns.h"
-#include "Widgets/ActiveSession/Replication/Client/Single/SSingleClientToolbar.h"
+#include "Widgets/ActiveSession/Replication/Client/SClientToolbar.h"
 
 #include "HAL/IConsoleManager.h"
 #include "Widgets/SBoxPanel.h"
@@ -37,7 +37,7 @@ namespace UE::MultiUserClient
 		FReplicationClient* ReplicationClient = GetReplicationClientAttribute.Get();
 		check(ReplicationClient);
 		
-		ConcertClientSharedSlate::IObjectToPropertiesModel& PropertyModel = *ReplicationClient->GetClientEditModel();
+		ConcertClientSharedSlate::IReplicationStreamModel& PropertyModel = *ReplicationClient->GetClientEditModel();
 		FGlobalAuthorityCache& AuthorityCache = ClientManager->GetAuthorityCache();
 		ChildSlot
 		[
@@ -48,7 +48,7 @@ namespace UE::MultiUserClient
 			.AutoHeight()
 			.Padding(2.f)
 			[
-				SNew(SSingleClientToolbar, PropertyModel, AuthorityCache)
+				SNew(SClientToolbar, PropertyModel, AuthorityCache)
 				.ViewSelectionArea() [ InArgs._ViewSelectionArea.Widget ]
 				.DisplayedClients(TSet{ ReplicationClient->GetEndpointId() })
 			]
@@ -99,7 +99,7 @@ namespace UE::MultiUserClient
 		)
 	{
 		using namespace ConcertClientSharedSlate;
-		IObjectToPropertiesModel& PropertyModel = *InReplicationClient.GetClientEditModel();
+		IReplicationStreamModel& PropertyModel = *InReplicationClient.GetClientEditModel();
 		FAuthorityChangeTracker& AuthorityTracker = InReplicationClient.GetAuthorityDiffer();
 		ISubmissionWorkflow& SubmissionWorkflow = InReplicationClient.GetSubmissionWorkflow();
 		FGlobalAuthorityCache& AuthorityCache = ClientManager->GetAuthorityCache();
@@ -149,7 +149,7 @@ namespace UE::MultiUserClient
 		)
 	{
 		using namespace ConcertClientSharedSlate;
-		IObjectToPropertiesModel& PropertyModel = *InReplicationClient.GetClientEditModel();
+		IReplicationStreamModel& PropertyModel = *InReplicationClient.GetClientEditModel();
 		FAuthorityChangeTracker& AuthorityTracker = InReplicationClient.GetAuthorityDiffer();
 		ISubmissionWorkflow& SubmissionWorkflow = InReplicationClient.GetSubmissionWorkflow();
 		FGlobalAuthorityCache& AuthorityCache = ClientManager->GetAuthorityCache();
diff --git a/Engine/Plugins/Developer/Concert/ConcertApp/MultiUserClient/Source/MultiUserClient/Private/Widgets/ActiveSession/Replication/Client/Single/SSingleClientToolbar.cpp b/Engine/Plugins/Developer/Concert/ConcertApp/MultiUserClient/Source/MultiUserClient/Private/Widgets/ActiveSession/Replication/Client/Single/SSingleClientToolbar.cpp
deleted file mode 100644
diff --git a/Engine/Plugins/Developer/Concert/ConcertApp/MultiUserClient/Source/MultiUserClient/Private/Widgets/ActiveSession/Replication/Client/Single/SSingleClientToolbar.h b/Engine/Plugins/Developer/Concert/ConcertApp/MultiUserClient/Source/MultiUserClient/Private/Widgets/ActiveSession/Replication/Client/Single/SSingleClientToolbar.h
deleted file mode 100644
diff --git a/Engine/Plugins/Developer/Concert/ConcertApp/MultiUserClient/Source/MultiUserReplicationEditor/Private/StreamEditor/ReplicationSessionPresetEditorToolkit.cpp b/Engine/Plugins/Developer/Concert/ConcertApp/MultiUserClient/Source/MultiUserReplicationEditor/Private/StreamEditor/ReplicationSessionPresetEditorToolkit.cpp
--- a/Engine/Plugins/Developer/Concert/ConcertApp/MultiUserClient/Source/MultiUserReplicationEditor/Private/StreamEditor/ReplicationSessionPresetEditorToolkit.cpp
+++ b/Engine/Plugins/Developer/Concert/ConcertApp/MultiUserClient/Source/MultiUserReplicationEditor/Private/StreamEditor/ReplicationSessionPresetEditorToolkit.cpp
@@ -74,7 +74,7 @@ namespace UE::MultiUserReplicationEditor
 	{
 		using namespace ConcertClientSharedSlate;
 		
-		const TSharedRef<IEditableObjectToPropertiesModel> AssetReadWriteModel = CreatePropertySelectionModel(
+		const TSharedRef<IEditableReplicationStreamModel> AssetReadWriteModel = CreatePropertySelectionModel(
 			*GetEditedStreamAsset(),
 			GetEditedStreamAsset()->GetUnassignedClient()->Stream->MakeReplicationMapGetterAttribute()
 			);
diff --git a/Engine/Plugins/Developer/Concert/ConcertUI/ConcertClientSharedSlate/Source/ConcertClientSharedSlate/Private/Replication/Editor/Model/GenericPropertySelectionModel.h b/Engine/Plugins/Developer/Concert/ConcertUI/ConcertClientSharedSlate/Source/ConcertClientSharedSlate/Private/Replication/Editor/Model/GenericPropertySelectionModel.h
--- a/Engine/Plugins/Developer/Concert/ConcertUI/ConcertClientSharedSlate/Source/ConcertClientSharedSlate/Private/Replication/Editor/Model/GenericPropertySelectionModel.h
+++ b/Engine/Plugins/Developer/Concert/ConcertUI/ConcertClientSharedSlate/Source/ConcertClientSharedSlate/Private/Replication/Editor/Model/GenericPropertySelectionModel.h
@@ -2,7 +2,7 @@
 
 #pragma once
 
-#include "Replication/Editor/Model/IEditableObjectToPropertiesModel.h"
+#include "Replication/Editor/Model/IEditableReplicationStreamModel.h"
 #include "Misc/Attribute.h"
 #include "UObject/WeakObjectPtrTemplates.h"
 
@@ -15,13 +15,13 @@ namespace UE::ConcertClientSharedSlate
 {
 	/** Implements logic for editing a FObjectReplicationMap contained in an UObject. */
 	class FGenericPropertySelectionModel
-		: public IEditableObjectToPropertiesModel
+		: public IEditableReplicationStreamModel
 	{
 	public:
 		
 		FGenericPropertySelectionModel(TAttribute<FObjectReplicationMap*> ReplicationMapAttribute);
 		
-		//~ Begin IObjectToPropertiesModel Interface
+		//~ Begin IReplicationStreamModel Interface
 		virtual uint32 GetNumReplicatedObjects() const override;
 		virtual uint32 GetNumProperties(const FSoftObjectPath& Object) const override;
 		virtual FSoftClassPath GetObjectClass(const FSoftObjectPath& Object) const override;
@@ -29,16 +29,16 @@ namespace UE::ConcertClientSharedSlate
 		virtual bool ContainsProperties(const FSoftObjectPath& Object, const TSet<FConcertPropertyChain>& Properties) const override;
 		virtual bool ForEachReplicatedObject(TFunctionRef<EBreakBehavior(const FSoftObjectPath& Object)> Delegate) const override;
 		virtual bool ForEachProperty(const FSoftObjectPath& Object, TFunctionRef<EBreakBehavior(const FConcertPropertyChain& Parent)> Delegate) const override;
-		//~ End IObjectToPropertiesModel Interface
+		//~ End IReplicationStreamModel Interface
 		
-		//~ Begin IEditableObjectToPropertiesModel Interface
+		//~ Begin IEditableReplicationStreamModel Interface
 		virtual void AddObjects(TConstArrayView<UObject*> Objects) override;
 		virtual void RemoveObjects(TConstArrayView<FSoftObjectPath> Objects) override;
 		virtual void AddProperties(const FSoftObjectPath&, TConstArrayView<FConcertPropertyChain> Properties) override;
 		virtual void RemoveProperties(const FSoftObjectPath&, TConstArrayView<FConcertPropertyChain> Properties) override;
 		virtual FOnObjectsChanged& OnObjectsChanged() override { return OnObjectsChangedDelegate; }
 		virtual FOnPropertiesChanged& OnPropertiesChanged() override { return OnPropertiesChangedDelegate; }
-		//~ End IEditableObjectToPropertiesModel Interface
+		//~ End IEditableReplicationStreamModel Interface
 
 	private:
 
diff --git a/Engine/Plugins/Developer/Concert/ConcertUI/ConcertClientSharedSlate/Source/ConcertClientSharedSlate/Private/Replication/Editor/Model/TransactionalPropertySelectionModel.h b/Engine/Plugins/Developer/Concert/ConcertUI/ConcertClientSharedSlate/Source/ConcertClientSharedSlate/Private/Replication/Editor/Model/TransactionalPropertySelectionModel.h
--- a/Engine/Plugins/Developer/Concert/ConcertUI/ConcertClientSharedSlate/Source/ConcertClientSharedSlate/Private/Replication/Editor/Model/TransactionalPropertySelectionModel.h
+++ b/Engine/Plugins/Developer/Concert/ConcertUI/ConcertClientSharedSlate/Source/ConcertClientSharedSlate/Private/Replication/Editor/Model/TransactionalPropertySelectionModel.h
@@ -19,12 +19,12 @@ namespace UE::ConcertClientSharedSlate
 			TAttribute<FObjectReplicationMap*> ReplicationMapAttribute
 			);
 		
-		//~ Begin IEditableObjectToPropertiesModel Interface
+		//~ Begin IEditableReplicationStreamModel Interface
 		virtual void AddObjects(TConstArrayView<UObject*> Objects) override;
 		virtual void RemoveObjects(TConstArrayView<FSoftObjectPath> Objects) override;
 		virtual void AddProperties(const FSoftObjectPath&, TConstArrayView<FConcertPropertyChain> Properties) override;
 		virtual void RemoveProperties(const FSoftObjectPath&, TConstArrayView<FConcertPropertyChain> Properties) override;
-		//~ End IEditableObjectToPropertiesModel Interface
+		//~ End IEditableReplicationStreamModel Interface
 
 		//~ Begin FEditorUndoClient Interface
 		virtual bool MatchesContext(const FTransactionContext& InContext, const TArray<TPair<UObject*, FTransactionObjectEvent>>& TransactionObjectContexts) const override;
diff --git a/Engine/Plugins/Developer/Concert/ConcertUI/ConcertClientSharedSlate/Source/ConcertClientSharedSlate/Private/Replication/Editor/View/DisplayUtils.cpp b/Engine/Plugins/Developer/Concert/ConcertUI/ConcertClientSharedSlate/Source/ConcertClientSharedSlate/Private/Replication/Editor/View/DisplayUtils.cpp
--- a/Engine/Plugins/Developer/Concert/ConcertUI/ConcertClientSharedSlate/Source/ConcertClientSharedSlate/Private/Replication/Editor/View/DisplayUtils.cpp
+++ b/Engine/Plugins/Developer/Concert/ConcertUI/ConcertClientSharedSlate/Source/ConcertClientSharedSlate/Private/Replication/Editor/View/DisplayUtils.cpp
@@ -2,7 +2,7 @@
 
 #include "Replication/Editor/View/DisplayUtils.h"
 
-#include "Replication/Editor/Model/IObjectToPropertiesModel.h"
+#include "Replication/Editor/Model/IReplicationStreamModel.h"
 
 #include "AssetRegistry/AssetData.h"
 #include "AssetRegistry/AssetRegistryModule.h"
@@ -45,7 +45,7 @@ namespace UE::ConcertClientSharedSlate::DisplayUtils
 		return Object.GetName();
 	}
 
-	FText GetObjectTypeText(const IObjectToPropertiesModel& Model, const FSoftObjectPath& Object)
+	FText GetObjectTypeText(const IReplicationStreamModel& Model, const FSoftObjectPath& Object)
 	{
 		if (const FSoftClassPath ClassPath = Model.GetObjectClass(Object); ClassPath.IsValid())
 		{
@@ -75,7 +75,7 @@ namespace UE::ConcertClientSharedSlate::DisplayUtils
 		return FText::GetEmpty();
 	}
 
-	FSlateIcon GetObjectIcon(const IObjectToPropertiesModel& Model, const FSoftObjectPath& Object)
+	FSlateIcon GetObjectIcon(const IReplicationStreamModel& Model, const FSoftObjectPath& Object)
 	{
 		if (const FSoftClassPath ClassPath = Model.GetObjectClass(Object); ClassPath.IsValid())
 		{
diff --git a/Engine/Plugins/Developer/Concert/ConcertUI/ConcertClientSharedSlate/Source/ConcertClientSharedSlate/Private/Replication/Editor/View/DisplayUtils.h b/Engine/Plugins/Developer/Concert/ConcertUI/ConcertClientSharedSlate/Source/ConcertClientSharedSlate/Private/Replication/Editor/View/DisplayUtils.h
--- a/Engine/Plugins/Developer/Concert/ConcertUI/ConcertClientSharedSlate/Source/ConcertClientSharedSlate/Private/Replication/Editor/View/DisplayUtils.h
+++ b/Engine/Plugins/Developer/Concert/ConcertUI/ConcertClientSharedSlate/Source/ConcertClientSharedSlate/Private/Replication/Editor/View/DisplayUtils.h
@@ -12,7 +12,7 @@ struct FSoftObjectPath;
 
 namespace UE::ConcertClientSharedSlate
 {
-	class IObjectToPropertiesModel;
+	class IReplicationStreamModel;
 }
 
 namespace UE::ConcertClientSharedSlate::DisplayUtils
@@ -26,10 +26,10 @@ namespace UE::ConcertClientSharedSlate::DisplayUtils
 	FString GetObjectDisplayString(const UObject& Object);
 	
 	/** @return The text to use for displaying this object's type */
-	FText GetObjectTypeText(const IObjectToPropertiesModel& Model, const FSoftObjectPath& Object);
+	FText GetObjectTypeText(const IReplicationStreamModel& Model, const FSoftObjectPath& Object);
 	
 	/** @return The icon to use for this object */
-	FSlateIcon GetObjectIcon(const IObjectToPropertiesModel& Model, const FSoftObjectPath& Object);
+	FSlateIcon GetObjectIcon(const IReplicationStreamModel& Model, const FSoftObjectPath& Object);
 	/** @return The icon to use for this object */
 	FSlateIcon GetObjectIcon(UObject& Object);
 
diff --git a/Engine/Plugins/Developer/Concert/ConcertUI/ConcertClientSharedSlate/Source/ConcertClientSharedSlate/Private/Replication/Editor/View/ObjectEditor/FakeObjectToPropertiesEditorModel.h b/Engine/Plugins/Developer/Concert/ConcertUI/ConcertClientSharedSlate/Source/ConcertClientSharedSlate/Private/Replication/Editor/View/ObjectEditor/FakeObjectToPropertiesEditorModel.h
--- a/Engine/Plugins/Developer/Concert/ConcertUI/ConcertClientSharedSlate/Source/ConcertClientSharedSlate/Private/Replication/Editor/View/ObjectEditor/FakeObjectToPropertiesEditorModel.h
+++ b/Engine/Plugins/Developer/Concert/ConcertUI/ConcertClientSharedSlate/Source/ConcertClientSharedSlate/Private/Replication/Editor/View/ObjectEditor/FakeObjectToPropertiesEditorModel.h
@@ -2,7 +2,7 @@
 
 #pragma once
 
-#include "Replication/Editor/Model/IObjectToPropertiesModel.h"
+#include "Replication/Editor/Model/IReplicationStreamModel.h"
 
 namespace UE::ConcertClientSharedSlate
 {
@@ -27,12 +27,12 @@ namespace UE::ConcertClientSharedSlate
 	 * underlying model.
 	 * @see UE::ConcertClientSharedSlate::ReplicationPropertyColumns::ReplicatesColumns.
 	 */
-	class FFakeObjectToPropertiesEditorModel : public IObjectToPropertiesModel
+	class FFakeObjectToPropertiesEditorModel : public IReplicationStreamModel
 	{
 	public:
 
 		FFakeObjectToPropertiesEditorModel(
-			TSharedRef<IObjectToPropertiesModel> RealModel,
+			TSharedRef<IReplicationStreamModel> RealModel,
 			TSharedRef<IPropertySelectionSourceModel> PropertySelectionSource,
 			EFakeObjectModelFlags Flags
 			)
@@ -44,7 +44,7 @@ namespace UE::ConcertClientSharedSlate
 		/** Whether this object should be displayed in the outliner. */
 		bool IsTopLevelObject(const FSoftObjectPath& ObjectPath) const;
 
-		//~ Begin IObjectToPropertiesModel Interface
+		//~ Begin IReplicationStreamModel Interface
 		// Technically these functions should be also be wrapped but the SObjectToPropertyView does not use them so let's not for now.
 		virtual uint32 GetNumReplicatedObjects() const override { return RealModel->GetNumReplicatedObjects(); }
 		virtual uint32 GetNumProperties(const FSoftObjectPath& Object) const override { return RealModel->GetNumProperties(Object); }
@@ -54,12 +54,12 @@ namespace UE::ConcertClientSharedSlate
 		virtual FSoftClassPath GetObjectClass(const FSoftObjectPath& Object) const override;
 		virtual bool ForEachReplicatedObject(TFunctionRef<EBreakBehavior(const FSoftObjectPath& Object)> Delegate) const override;
 		virtual bool ForEachProperty(const FSoftObjectPath& ObjectPath, TFunctionRef<EBreakBehavior(const FConcertPropertyChain& Property)> Delegate) const override;
-		//~ End IObjectToPropertiesModel Interface
+		//~ End IReplicationStreamModel Interface
 
 	private:
 
 		/** The real model which is used to implement all the other functions. */
-		const TSharedRef<IObjectToPropertiesModel> RealModel;
+		const TSharedRef<IReplicationStreamModel> RealModel;
 
 		/** Determines the properties that can be selected. */
 		const TSharedRef<IPropertySelectionSourceModel> PropertySelectionSource;
diff --git a/Engine/Plugins/Developer/Concert/ConcertUI/ConcertClientSharedSlate/Source/ConcertClientSharedSlate/Private/Replication/Editor/View/ObjectEditor/SBaseReplicationStreamEditor.cpp b/Engine/Plugins/Developer/Concert/ConcertUI/ConcertClientSharedSlate/Source/ConcertClientSharedSlate/Private/Replication/Editor/View/ObjectEditor/SBaseReplicationStreamEditor.cpp
--- a/Engine/Plugins/Developer/Concert/ConcertUI/ConcertClientSharedSlate/Source/ConcertClientSharedSlate/Private/Replication/Editor/View/ObjectEditor/SBaseReplicationStreamEditor.cpp
+++ b/Engine/Plugins/Developer/Concert/ConcertUI/ConcertClientSharedSlate/Source/ConcertClientSharedSlate/Private/Replication/Editor/View/ObjectEditor/SBaseReplicationStreamEditor.cpp
@@ -4,7 +4,7 @@
 
 #include "FakeObjectToPropertiesEditorModel.h"
 #include "Model/Item/SourceModelBuilders.h"
-#include "Replication/Editor/Model/IEditableObjectToPropertiesModel.h"
+#include "Replication/Editor/Model/IEditableReplicationStreamModel.h"
 #include "Replication/Editor/Model/Object/IObjectSelectionSourceModel.h"
 #include "Replication/Editor/Model/Property/IPropertySelectionSourceModel.h"
 #include "Replication/Editor/Model/ReplicatedObjectData.h"
@@ -23,7 +23,7 @@ namespace UE::ConcertClientSharedSlate
 {
 	void SBaseReplicationStreamEditor::Construct(
 		const FArguments& InArgs,
-		TSharedRef<IEditableObjectToPropertiesModel> InPropertiesModel,
+		TSharedRef<IEditableReplicationStreamModel> InPropertiesModel,
 		TSharedRef<IObjectSelectionSourceModel> InObjectSelectionSource,
 		TSharedRef<IPropertySelectionSourceModel> InPropertySelectionSource)
 	{
diff --git a/Engine/Plugins/Developer/Concert/ConcertUI/ConcertClientSharedSlate/Source/ConcertClientSharedSlate/Private/Replication/Editor/View/ObjectEditor/SBaseReplicationStreamEditor.h b/Engine/Plugins/Developer/Concert/ConcertUI/ConcertClientSharedSlate/Source/ConcertClientSharedSlate/Private/Replication/Editor/View/ObjectEditor/SBaseReplicationStreamEditor.h
--- a/Engine/Plugins/Developer/Concert/ConcertUI/ConcertClientSharedSlate/Source/ConcertClientSharedSlate/Private/Replication/Editor/View/ObjectEditor/SBaseReplicationStreamEditor.h
+++ b/Engine/Plugins/Developer/Concert/ConcertUI/ConcertClientSharedSlate/Source/ConcertClientSharedSlate/Private/Replication/Editor/View/ObjectEditor/SBaseReplicationStreamEditor.h
@@ -15,13 +15,13 @@ struct FConcertReplicationEditorSettings;
 
 namespace UE::ConcertClientSharedSlate
 {
-	class IEditableObjectToPropertiesModel;
+	class IEditableReplicationStreamModel;
 	class FFakeObjectToPropertiesEditorModel;
 	class FReplicatedObjectData;
 	class FReplicatedPropertyData;
 	class IReplicationSubobjectView;
 	class IObjectSelectionSourceModel;
-	class IObjectToPropertiesModel;
+	class IReplicationStreamModel;
 	class IPropertySelectionSourceModel;
 	class IReplicationSubobjectView;
 	class ISubobjectModel;
@@ -31,11 +31,11 @@ namespace UE::ConcertClientSharedSlate
 	struct FSelectableObjectInfo;
 	
 	/**
-	 * Extends SObjectToPropertyView with common functionality for editing (with an IEditableObjectToPropertiesModel):
-	 * - Pressing delete calls IEditableObjectToPropertiesModel::RemoveObjects
+	 * Extends SObjectToPropertyView with common functionality for editing (with an IEditableReplicationStreamModel):
+	 * - Pressing delete calls IEditableReplicationStreamModel::RemoveObjects
 	 * - A combo box to left of the object search bar for adding new top level objects
 	 * - Default context menu items for top-level objects, such as deleting or adding more based right-clicked object
-	 * - Wraps the IEditableObjectToPropertiesModel with FFakeObjectToPropertiesEditorModel so the SObjectToPropertyView:
+	 * - Wraps the IEditableReplicationStreamModel with FFakeObjectToPropertiesEditorModel so the SObjectToPropertyView:
 	 *		- displays only top-level objects in the outliner (i.e. excludes components)
 	 *		- displays all properties in the property view instead of just those registered in the model
 	 *
@@ -82,7 +82,7 @@ namespace UE::ConcertClientSharedSlate
 		SLATE_END_ARGS()
 
 		void Construct(const FArguments& InArgs,
-           TSharedRef<IEditableObjectToPropertiesModel> InPropertiesModel,
+           TSharedRef<IEditableReplicationStreamModel> InPropertiesModel,
            TSharedRef<IObjectSelectionSourceModel> InObjectSelectionSource,
            TSharedRef<IPropertySelectionSourceModel> InPropertySelectionSource
 		);
@@ -100,7 +100,7 @@ namespace UE::ConcertClientSharedSlate
 		TSharedPtr<SReplicationStreamViewer> ReplicationViewer;
 
 		/** For reading and writting to the edited asset */
-		TSharedPtr<IEditableObjectToPropertiesModel> EditablePropertiesModel;
+		TSharedPtr<IEditableReplicationStreamModel> EditablePropertiesModel;
 		/**
 		 * Fakes to SObjectToPropertyView that all UClass properties are contained.
 		 * We inject checkboxes to SObjectToPropertyView which do the actual adding and removing.
diff --git a/Engine/Plugins/Developer/Concert/ConcertUI/ConcertClientSharedSlate/Source/ConcertClientSharedSlate/Private/Replication/Editor/View/ObjectEditor/SDefaultReplicationStreamEditor.cpp b/Engine/Plugins/Developer/Concert/ConcertUI/ConcertClientSharedSlate/Source/ConcertClientSharedSlate/Private/Replication/Editor/View/ObjectEditor/SDefaultReplicationStreamEditor.cpp
--- a/Engine/Plugins/Developer/Concert/ConcertUI/ConcertClientSharedSlate/Source/ConcertClientSharedSlate/Private/Replication/Editor/View/ObjectEditor/SDefaultReplicationStreamEditor.cpp
+++ b/Engine/Plugins/Developer/Concert/ConcertUI/ConcertClientSharedSlate/Source/ConcertClientSharedSlate/Private/Replication/Editor/View/ObjectEditor/SDefaultReplicationStreamEditor.cpp
@@ -2,7 +2,7 @@
 
 #include "SDefaultReplicationStreamEditor.h"
 
-#include "Replication/Editor/Model/IEditableObjectToPropertiesModel.h"
+#include "Replication/Editor/Model/IEditableReplicationStreamModel.h"
 #include "Replication/Editor/View/ObjectViewer/Tree/SelectionViewerColumns.h"
 
 #include "SBaseReplicationStreamEditor.h"
@@ -11,7 +11,7 @@ namespace UE::ConcertClientSharedSlate
 {
 	void SDefaultReplicationStreamEditor::Construct(
 		const FArguments& InArgs,
-		TSharedRef<IEditableObjectToPropertiesModel> InPropertiesModel,
+		TSharedRef<IEditableReplicationStreamModel> InPropertiesModel,
 		TSharedRef<IObjectSelectionSourceModel> InObjectSelectionSource,
 		TSharedRef<IPropertySelectionSourceModel> InPropertySelectionSource)
 	{
diff --git a/Engine/Plugins/Developer/Concert/ConcertUI/ConcertClientSharedSlate/Source/ConcertClientSharedSlate/Private/Replication/Editor/View/ObjectEditor/SDefaultReplicationStreamEditor.h b/Engine/Plugins/Developer/Concert/ConcertUI/ConcertClientSharedSlate/Source/ConcertClientSharedSlate/Private/Replication/Editor/View/ObjectEditor/SDefaultReplicationStreamEditor.h
--- a/Engine/Plugins/Developer/Concert/ConcertUI/ConcertClientSharedSlate/Source/ConcertClientSharedSlate/Private/Replication/Editor/View/ObjectEditor/SDefaultReplicationStreamEditor.h
+++ b/Engine/Plugins/Developer/Concert/ConcertUI/ConcertClientSharedSlate/Source/ConcertClientSharedSlate/Private/Replication/Editor/View/ObjectEditor/SDefaultReplicationStreamEditor.h
@@ -14,7 +14,7 @@ struct FConcertReplicationEditorSettings;
 
 namespace UE::ConcertClientSharedSlate
 {
-	class IEditableObjectToPropertiesModel;
+	class IEditableReplicationStreamModel;
 	class IObjectSelectionSourceModel;
 	class IPropertySelectionSourceModel;
 	class IReplicationSubobjectView;
@@ -65,7 +65,7 @@ namespace UE::ConcertClientSharedSlate
 
 		void Construct(
 			const FArguments& InArgs,
-			TSharedRef<IEditableObjectToPropertiesModel> InPropertiesModel,
+			TSharedRef<IEditableReplicationStreamModel> InPropertiesModel,
 			TSharedRef<IObjectSelectionSourceModel> InObjectSelectionSource,
 			TSharedRef<IPropertySelectionSourceModel> InPropertySelectionSource
 			);
@@ -82,7 +82,7 @@ namespace UE::ConcertClientSharedSlate
 		TSharedPtr<SBaseReplicationStreamEditor> WrappedEditor;
 		
 		/** Model needed for sorting properties based on selection. */
-		TSharedPtr<IEditableObjectToPropertiesModel> PropertiesModel;
+		TSharedPtr<IEditableReplicationStreamModel> PropertiesModel;
 		
 		bool SortPropertiesPredicate(const FReplicatedPropertyData& Left, const FReplicatedPropertyData& Right) const;
     };
diff --git a/Engine/Plugins/Developer/Concert/ConcertUI/ConcertClientSharedSlate/Source/ConcertClientSharedSlate/Private/Replication/Editor/View/ObjectViewer/SReplicationStreamViewer.cpp b/Engine/Plugins/Developer/Concert/ConcertUI/ConcertClientSharedSlate/Source/ConcertClientSharedSlate/Private/Replication/Editor/View/ObjectViewer/SReplicationStreamViewer.cpp
--- a/Engine/Plugins/Developer/Concert/ConcertUI/ConcertClientSharedSlate/Source/ConcertClientSharedSlate/Private/Replication/Editor/View/ObjectViewer/SReplicationStreamViewer.cpp
+++ b/Engine/Plugins/Developer/Concert/ConcertUI/ConcertClientSharedSlate/Source/ConcertClientSharedSlate/Private/Replication/Editor/View/ObjectViewer/SReplicationStreamViewer.cpp
@@ -3,7 +3,7 @@
 #include "SReplicationStreamViewer.h"
 
 #include "ConcertFrontendUtils.h"
-#include "Replication/Editor/Model/IObjectToPropertiesModel.h"
+#include "Replication/Editor/Model/IReplicationStreamModel.h"
 #include "Replication/Editor/Model/ISubobjectModel.h"
 #include "Replication/Editor/Model/ReplicatedObjectData.h"
 #include "Replication/Editor/View/ObjectViewer/Property/SReplicatedPropertiesView.h"
@@ -20,7 +20,7 @@
 
 namespace UE::ConcertClientSharedSlate
 {
-	void SReplicationStreamViewer::Construct(const FArguments& InArgs, TSharedRef<IObjectToPropertiesModel> InPropertiesModel)
+	void SReplicationStreamViewer::Construct(const FArguments& InArgs, TSharedRef<IReplicationStreamModel> InPropertiesModel)
 	{
 		PropertiesModel = MoveTemp(InPropertiesModel);
 		SubobjectModel = InArgs._SubobjectModel;
diff --git a/Engine/Plugins/Developer/Concert/ConcertUI/ConcertClientSharedSlate/Source/ConcertClientSharedSlate/Private/Replication/Editor/View/ObjectViewer/SReplicationStreamViewer.h b/Engine/Plugins/Developer/Concert/ConcertUI/ConcertClientSharedSlate/Source/ConcertClientSharedSlate/Private/Replication/Editor/View/ObjectViewer/SReplicationStreamViewer.h
--- a/Engine/Plugins/Developer/Concert/ConcertUI/ConcertClientSharedSlate/Source/ConcertClientSharedSlate/Private/Replication/Editor/View/ObjectViewer/SReplicationStreamViewer.h
+++ b/Engine/Plugins/Developer/Concert/ConcertUI/ConcertClientSharedSlate/Source/ConcertClientSharedSlate/Private/Replication/Editor/View/ObjectViewer/SReplicationStreamViewer.h
@@ -24,18 +24,18 @@ namespace UE::ConcertClientSharedSlate
 	class SSubobjectAndPropertySection;
 	class FReplicatedPropertyData;
 	class FReplicatedObjectData;
-	class IEditableObjectToPropertiesModel;
+	class IEditableReplicationStreamModel;
 	class IReplicationSubobjectView;
-	class IObjectToPropertiesModel;
+	class IReplicationStreamModel;
 	class ISubobjectModel;
 	class SReplicatedPropertiesView;
 	
 	/**
 	 * Root widget for viewing UMultiUserPropertyReplicationSelection.
-	 * This widget knows how to display IObjectToPropertiesModel.
+	 * This widget knows how to display IReplicationStreamModel.
 	 * 
 	 * The underlying data is modified by SObjectToPropertyEditor, which uses this widget's extension
-	 * points to call functions on IEditableObjectToPropertiesModel.
+	 * points to call functions on IEditableReplicationStreamModel.
 	 *
 	 * Important: this view should be possible to be built in programs, so it should not reference things like AActor,
 	 * UActorComponent, ResolveObject, etc. directly. 
@@ -74,7 +74,7 @@ namespace UE::ConcertClientSharedSlate
 			SLATE_ATTRIBUTE(FText, NoOutlinerObjects)
 		SLATE_END_ARGS()
 
-		void Construct(const FArguments& InArgs, TSharedRef<IObjectToPropertiesModel> InPropertiesModel);
+		void Construct(const FArguments& InArgs, TSharedRef<IReplicationStreamModel> InPropertiesModel);
 
 		//~ Begin IReplicationStreamViewer Interface
 		virtual void Refresh() override;
@@ -101,7 +101,7 @@ namespace UE::ConcertClientSharedSlate
 	private:
 
 		/** The model this view is visualizing. */
-		TSharedPtr<IObjectToPropertiesModel> PropertiesModel;
+		TSharedPtr<IReplicationStreamModel> PropertiesModel;
 		/** Can be null. If set, this determines the children nested under the root objects. */
 		TSharedPtr<ISubobjectModel> SubobjectModel;
 
diff --git a/Engine/Plugins/Developer/Concert/ConcertUI/ConcertClientSharedSlate/Source/ConcertClientSharedSlate/Private/Replication/Editor/View/ObjectViewer/SSubobjectAndPropertySection.cpp b/Engine/Plugins/Developer/Concert/ConcertUI/ConcertClientSharedSlate/Source/ConcertClientSharedSlate/Private/Replication/Editor/View/ObjectViewer/SSubobjectAndPropertySection.cpp
--- a/Engine/Plugins/Developer/Concert/ConcertUI/ConcertClientSharedSlate/Source/ConcertClientSharedSlate/Private/Replication/Editor/View/ObjectViewer/SSubobjectAndPropertySection.cpp
+++ b/Engine/Plugins/Developer/Concert/ConcertUI/ConcertClientSharedSlate/Source/ConcertClientSharedSlate/Private/Replication/Editor/View/ObjectViewer/SSubobjectAndPropertySection.cpp
@@ -2,7 +2,7 @@
 
 #include "SSubobjectAndPropertySection.h"
 
-#include "Replication/Editor/Model/IObjectToPropertiesModel.h"
+#include "Replication/Editor/Model/IReplicationStreamModel.h"
 #include "Replication/Editor/Model/ReplicatedObjectData.h"
 #include "Replication/Editor/Model/ReplicatedPropertyData.h"
 #include "Replication/Editor/View/IReplicationSubobjectView.h"
@@ -21,7 +21,7 @@
 
 namespace UE::ConcertClientSharedSlate
 {
-	void SSubobjectAndPropertySection::Construct(const FArguments& InArgs, TSharedRef<IObjectToPropertiesModel> InPropertiesModel)
+	void SSubobjectAndPropertySection::Construct(const FArguments& InArgs, TSharedRef<IReplicationStreamModel> InPropertiesModel)
 	{
 		PropertiesModel = MoveTemp(InPropertiesModel);
 		SortPropertyRowPredicate = InArgs._SortPropertyRowPredicate;
diff --git a/Engine/Plugins/Developer/Concert/ConcertUI/ConcertClientSharedSlate/Source/ConcertClientSharedSlate/Private/Replication/Editor/View/ObjectViewer/SSubobjectAndPropertySection.h b/Engine/Plugins/Developer/Concert/ConcertUI/ConcertClientSharedSlate/Source/ConcertClientSharedSlate/Private/Replication/Editor/View/ObjectViewer/SSubobjectAndPropertySection.h
--- a/Engine/Plugins/Developer/Concert/ConcertUI/ConcertClientSharedSlate/Source/ConcertClientSharedSlate/Private/Replication/Editor/View/ObjectViewer/SSubobjectAndPropertySection.h
+++ b/Engine/Plugins/Developer/Concert/ConcertUI/ConcertClientSharedSlate/Source/ConcertClientSharedSlate/Private/Replication/Editor/View/ObjectViewer/SSubobjectAndPropertySection.h
@@ -45,7 +45,7 @@ namespace UE::ConcertClientSharedSlate
 			SLATE_NAMED_SLOT(FArguments, LeftOfPropertySearchBar)
 		SLATE_END_ARGS()
 
-		void Construct(const FArguments& InArgs, TSharedRef<IObjectToPropertiesModel> InPropertiesModel);
+		void Construct(const FArguments& InArgs, TSharedRef<IReplicationStreamModel> InPropertiesModel);
 
 		/** If there is a subobject view, makes it select the root objects. */
 		void SelectRootObjects() const;
@@ -61,7 +61,7 @@ namespace UE::ConcertClientSharedSlate
 	private:
 		
 		/** The model this view is visualizing. */
-		TSharedPtr<IObjectToPropertiesModel> PropertiesModel;
+		TSharedPtr<IReplicationStreamModel> PropertiesModel;
 		
 		/** Optional. External widget that selects subobjects from ActorArea. */
 		TSharedPtr<IReplicationSubobjectView> SubobjectView;
diff --git a/Engine/Plugins/Developer/Concert/ConcertUI/ConcertClientSharedSlate/Source/ConcertClientSharedSlate/Private/Replication/Editor/View/ObjectViewer/Tree/SelectionViewerColumns.cpp b/Engine/Plugins/Developer/Concert/ConcertUI/ConcertClientSharedSlate/Source/ConcertClientSharedSlate/Private/Replication/Editor/View/ObjectViewer/Tree/SelectionViewerColumns.cpp
--- a/Engine/Plugins/Developer/Concert/ConcertUI/ConcertClientSharedSlate/Source/ConcertClientSharedSlate/Private/Replication/Editor/View/ObjectViewer/Tree/SelectionViewerColumns.cpp
+++ b/Engine/Plugins/Developer/Concert/ConcertUI/ConcertClientSharedSlate/Source/ConcertClientSharedSlate/Private/Replication/Editor/View/ObjectViewer/Tree/SelectionViewerColumns.cpp
@@ -3,7 +3,7 @@
 #include "SelectionViewerColumns.h"
 
 #include "ClassIconFinder.h"
-#include "Replication/Editor/Model/IEditableObjectToPropertiesModel.h"
+#include "Replication/Editor/Model/IEditableReplicationStreamModel.h"
 #include "Replication/Editor/Model/ISubobjectModel.h"
 #include "Replication/Editor/Model/ReplicatedPropertyData.h"
 #include "Replication/Editor/Model/ReplicatedObjectData.h"
@@ -29,7 +29,7 @@ namespace UE::ConcertClientSharedSlate::ReplicationColumns::TopLevel
 	const FName LabelColumnId = TEXT("LabelColumn");
 	const FName TypeColumnId = TEXT("TypeColumn");
 	
-	FReplicationTopLevelObjectColumn LabelColumn(TSharedRef<IObjectToPropertiesModel> Model, ISubobjectModel* SubobjectModel)
+	FReplicationTopLevelObjectColumn LabelColumn(TSharedRef<IReplicationStreamModel> Model, ISubobjectModel* SubobjectModel)
 	{
 		return FReplicationTopLevelObjectColumn(
 			FReplicationTopLevelObjectColumn::FArguments()
@@ -77,7 +77,7 @@ namespace UE::ConcertClientSharedSlate::ReplicationColumns::TopLevel
 			);
 	}
 	
-	FReplicationTopLevelObjectColumn TypeColumn(TSharedRef<IObjectToPropertiesModel> Model)
+	FReplicationTopLevelObjectColumn TypeColumn(TSharedRef<IReplicationStreamModel> Model)
 	{
 		return FReplicationTopLevelObjectColumn(
 			FReplicationTopLevelObjectColumn::FArguments()
@@ -213,7 +213,7 @@ namespace UE::ConcertClientSharedSlate::ReplicationColumns::Property
 		static ECheckBoxState OnGetPropertyCheckboxState(
 			const FConcertPropertyChain& PropertyChain,
 			const IReplicationStreamViewer& Viewer,
-			const IObjectToPropertiesModel& Model
+			const IReplicationStreamModel& Model
 			)
 		{
 			const TArray<FSoftObjectPath> SelectedObjectPaths = Viewer.GetObjectsBeingPropertyEdited();
@@ -224,7 +224,7 @@ namespace UE::ConcertClientSharedSlate::ReplicationColumns::Property
 			bool bIsChecked,
 			const FConcertPropertyChain& PropertyChain,
 			const IReplicationStreamViewer& Viewer,
-			IEditableObjectToPropertiesModel& Model
+			IEditableReplicationStreamModel& Model
 			)
 		{
 			const TArray Properties{ PropertyChain };
@@ -273,7 +273,7 @@ namespace UE::ConcertClientSharedSlate::ReplicationColumns::Property
 	
 	FReplicationPropertyColumn ReplicatesColumns(
 		TWeakPtr<IReplicationStreamViewer> Viewer,
-		TWeakPtr<IEditableObjectToPropertiesModel> Model,
+		TWeakPtr<IEditableReplicationStreamModel> Model,
 		TReplicationColumnDelegates<FReplicatedPropertyData>::FIsEnabled IsEnabledDelegate,
 		TAttribute<FText> DisabledToolTipText,
 		const float ColumnWidth,
@@ -288,14 +288,14 @@ namespace UE::ConcertClientSharedSlate::ReplicationColumns::Property
 				[Viewer, Model](const FReplicatedPropertyData& Data)
 				{
 					const TSharedPtr<IReplicationStreamViewer> ViewerPin = Viewer.Pin();
-					const TSharedPtr<IEditableObjectToPropertiesModel> ModelPin = Model.Pin();
+					const TSharedPtr<IEditableReplicationStreamModel> ModelPin = Model.Pin();
 					return ensure(ViewerPin && ModelPin) ? Private::OnGetPropertyCheckboxState(Data.GetProperty(), *ViewerPin, *ModelPin) : ECheckBoxState::Undetermined;
 				}),
 				FPropertyColumnDelegates::FOnColumnCheckboxChanged::CreateLambda(
 				[Viewer, Model](bool bIsChecked, const FReplicatedPropertyData& Data)
 				{
 					const TSharedPtr<IReplicationStreamViewer> ViewerPin = Viewer.Pin();
-					const TSharedPtr<IEditableObjectToPropertiesModel> ModelPin = Model.Pin();
+					const TSharedPtr<IEditableReplicationStreamModel> ModelPin = Model.Pin();
 					if (ensure(ViewerPin && ModelPin))
 					{
 						Private::OnPropertyCheckboxChanged(bIsChecked, Data.GetProperty(), *ViewerPin, *ModelPin);
@@ -318,7 +318,7 @@ namespace UE::ConcertClientSharedSlate::ReplicationColumns::Property
 		);
 	}
 
-	ECheckBoxState GetPropertyCheckboxStateBasedOnSelection(const FConcertPropertyChain& Property, TConstArrayView<FSoftObjectPath> Selection, const IObjectToPropertiesModel& Model)
+	ECheckBoxState GetPropertyCheckboxStateBasedOnSelection(const FConcertPropertyChain& Property, TConstArrayView<FSoftObjectPath> Selection, const IReplicationStreamModel& Model)
 	{
 		ECheckBoxState CheckBoxState = ECheckBoxState::Undetermined;
 		for (const FSoftObjectPath& SelectedObject : Selection)
@@ -343,7 +343,7 @@ namespace UE::ConcertClientSharedSlate::ReplicationColumns::Property
 	
 	bool SortBySelectionThenByName_PropertyPredicate(
 		const TArray<FSoftObjectPath>& SelectedObjects,
-		const IObjectToPropertiesModel& Model,
+		const IReplicationStreamModel& Model,
 		const FReplicatedPropertyData& Left,
 		const FReplicatedPropertyData& Right
 		)
diff --git a/Engine/Plugins/Developer/Concert/ConcertUI/ConcertClientSharedSlate/Source/ConcertClientSharedSlate/Private/Replication/Editor/View/ObjectViewer/Tree/SelectionViewerColumns.h b/Engine/Plugins/Developer/Concert/ConcertUI/ConcertClientSharedSlate/Source/ConcertClientSharedSlate/Private/Replication/Editor/View/ObjectViewer/Tree/SelectionViewerColumns.h
--- a/Engine/Plugins/Developer/Concert/ConcertUI/ConcertClientSharedSlate/Source/ConcertClientSharedSlate/Private/Replication/Editor/View/ObjectViewer/Tree/SelectionViewerColumns.h
+++ b/Engine/Plugins/Developer/Concert/ConcertUI/ConcertClientSharedSlate/Source/ConcertClientSharedSlate/Private/Replication/Editor/View/ObjectViewer/Tree/SelectionViewerColumns.h
@@ -8,9 +8,9 @@
 
 namespace UE::ConcertClientSharedSlate
 {
-	class IEditableObjectToPropertiesModel;
+	class IEditableReplicationStreamModel;
 	class IReplicationStreamViewer;
-	class IObjectToPropertiesModel;
+	class IReplicationStreamModel;
 	class ISubobjectModel;
 }
 
@@ -28,8 +28,8 @@ namespace UE::ConcertClientSharedSlate::ReplicationColumns::TopLevel
 		Type = 30,
 	};
 
-	FReplicationTopLevelObjectColumn LabelColumn(TSharedRef<IObjectToPropertiesModel> Model, ISubobjectModel* SubobjectModel = nullptr);
-	FReplicationTopLevelObjectColumn TypeColumn(TSharedRef<IObjectToPropertiesModel> Model);
+	FReplicationTopLevelObjectColumn LabelColumn(TSharedRef<IReplicationStreamModel> Model, ISubobjectModel* SubobjectModel = nullptr);
+	FReplicationTopLevelObjectColumn TypeColumn(TSharedRef<IReplicationStreamModel> Model);
 }
 
 namespace UE::ConcertClientSharedSlate::ReplicationColumns::Subobject
@@ -80,7 +80,7 @@ namespace UE::ConcertClientSharedSlate::ReplicationColumns::Property
 	 */
 	CONCERTCLIENTSHAREDSLATE_API FReplicationPropertyColumn ReplicatesColumns(
 		TWeakPtr<IReplicationStreamViewer> Viewer,
-		TWeakPtr<IEditableObjectToPropertiesModel> Model,
+		TWeakPtr<IEditableReplicationStreamModel> Model,
 		TReplicationColumnDelegates<FReplicatedPropertyData>::FIsEnabled IsEnabledDelegate = {},
 		TAttribute<FText> DisabledToolTipText = {},
 		const float ColumnWidth = 20.f,
@@ -96,12 +96,12 @@ namespace UE::ConcertClientSharedSlate::ReplicationColumns::Property
 	CONCERTCLIENTSHAREDSLATE_API ECheckBoxState GetPropertyCheckboxStateBasedOnSelection(
 		const FConcertPropertyChain& Property,
 		TConstArrayView<FSoftObjectPath> Selection,
-		const IObjectToPropertiesModel& Model
+		const IReplicationStreamModel& Model
 		);
 	/** Util for sorting property data by whether its ReplicatesColumns() column is checked. */
 	CONCERTCLIENTSHAREDSLATE_API bool SortBySelectionThenByName_PropertyPredicate(
 		const TArray<FSoftObjectPath>& SelectedObjects,
-		const IObjectToPropertiesModel& Model,
+		const IReplicationStreamModel& Model,
 		const FReplicatedPropertyData& Left,
 		const FReplicatedPropertyData& Right
 		);
diff --git a/Engine/Plugins/Developer/Concert/ConcertUI/ConcertClientSharedSlate/Source/ConcertClientSharedSlate/Private/Replication/ReplicationWidgetFactories.cpp b/Engine/Plugins/Developer/Concert/ConcertUI/ConcertClientSharedSlate/Source/ConcertClientSharedSlate/Private/Replication/ReplicationWidgetFactories.cpp
--- a/Engine/Plugins/Developer/Concert/ConcertUI/ConcertClientSharedSlate/Source/ConcertClientSharedSlate/Private/Replication/ReplicationWidgetFactories.cpp
+++ b/Engine/Plugins/Developer/Concert/ConcertUI/ConcertClientSharedSlate/Source/ConcertClientSharedSlate/Private/Replication/ReplicationWidgetFactories.cpp
@@ -53,7 +53,7 @@ namespace UE::ConcertClientSharedSlate
 			.ReplicationSettings(Params.ReplicationSettingsAttribute);
 	}
 
-	TSharedRef<IEditableObjectToPropertiesModel> CreatePropertySelectionModel(
+	TSharedRef<IEditableReplicationStreamModel> CreatePropertySelectionModel(
 		UObject& OwnerObject,
 		TAttribute<FObjectReplicationMap*> ReplicationMapAttribute
 		)
diff --git a/Engine/Plugins/Developer/Concert/ConcertUI/ConcertClientSharedSlate/Source/ConcertClientSharedSlate/Public/Replication/Editor/Model/IEditableObjectToPropertiesModel.h b/Engine/Plugins/Developer/Concert/ConcertUI/ConcertClientSharedSlate/Source/ConcertClientSharedSlate/Public/Replication/Editor/Model/IEditableObjectToPropertiesModel.h
deleted file mode 100644
diff --git a/Engine/Plugins/Developer/Concert/ConcertUI/ConcertClientSharedSlate/Source/ConcertClientSharedSlate/Public/Replication/Editor/Model/IEditableReplicationStreamModel.h b/Engine/Plugins/Developer/Concert/ConcertUI/ConcertClientSharedSlate/Source/ConcertClientSharedSlate/Public/Replication/Editor/Model/IEditableReplicationStreamModel.h
new file mode 100644
--- /dev/null
+++ b/Engine/Plugins/Developer/Concert/ConcertUI/ConcertClientSharedSlate/Source/ConcertClientSharedSlate/Public/Replication/Editor/Model/IEditableReplicationStreamModel.h
@@ -0,0 +1,51 @@
+// Copyright Epic Games, Inc. All Rights Reserved.
+
+#pragma once
+
+#include "IReplicationStreamModel.h"
+
+#include "Containers/ContainersFwd.h"
+#include "Containers/Map.h"
+#include "Containers/Set.h"
+#include "Delegates/DelegateCombinations.h"
+
+class UObject;
+struct FConcertPropertyChain;
+
+namespace UE::ConcertClientSharedSlate
+{
+	enum class EReplicatedObjectChangeReason : uint8
+	{
+		ChangedDirectly,
+		Transacted
+	};
+	
+	/**
+	 * Abstracts the concept of mapping objects to properties. This allows writing.
+	 * 
+	 * Models may not always be writable. When editing a UAsset, it will be writable. However, if we join a multi-user
+	 * session we do not want to edit the objects nor properties - only read.
+	 */
+	class CONCERTCLIENTSHAREDSLATE_API IEditableReplicationStreamModel : public IReplicationStreamModel
+	{
+	public:
+
+		/** Adds the objects to the mapping */
+		virtual void AddObjects(TConstArrayView<UObject*> Objects) = 0;
+		/** Removes the objects from the mapping */
+		virtual void RemoveObjects(TConstArrayView<FSoftObjectPath> Objects) = 0;
+
+		/** Adds these properties to the object's list of selected properties. */
+		virtual void AddProperties(const FSoftObjectPath& Object, TConstArrayView<FConcertPropertyChain> Properties) = 0;
+		/** Removes these properties from the object's list of selected properties. */
+		virtual void RemoveProperties(const FSoftObjectPath& Object, TConstArrayView<FConcertPropertyChain> Properties) = 0;
+
+		/** Called when the object list changes. AddedObjects and RemovedObjects are empty if and only if ChangeReason == EReplicatedObjectChangeReason::Transacted. */
+		DECLARE_MULTICAST_DELEGATE_ThreeParams(FOnObjectsChanged, TConstArrayView<UObject*> AddedObjects, TConstArrayView<FSoftObjectPath> RemovedObjects, EReplicatedObjectChangeReason ChangeReason);
+		virtual FOnObjectsChanged& OnObjectsChanged() = 0;
+
+		/** Called when the property list of some object changes. If the properties are removed because the object is removed outright, FOnObjectsChanged is called instead.  */
+		DECLARE_MULTICAST_DELEGATE(FOnPropertiesChanged);
+		virtual FOnPropertiesChanged& OnPropertiesChanged() = 0;
+	};
+}
\ No newline at end of file
diff --git a/Engine/Plugins/Developer/Concert/ConcertUI/ConcertClientSharedSlate/Source/ConcertClientSharedSlate/Public/Replication/Editor/Model/IObjectToPropertiesModel.h b/Engine/Plugins/Developer/Concert/ConcertUI/ConcertClientSharedSlate/Source/ConcertClientSharedSlate/Public/Replication/Editor/Model/IObjectToPropertiesModel.h
deleted file mode 100644
diff --git a/Engine/Plugins/Developer/Concert/ConcertUI/ConcertClientSharedSlate/Source/ConcertClientSharedSlate/Public/Replication/Editor/Model/IReplicationStreamModel.h b/Engine/Plugins/Developer/Concert/ConcertUI/ConcertClientSharedSlate/Source/ConcertClientSharedSlate/Public/Replication/Editor/Model/IReplicationStreamModel.h
new file mode 100644
--- /dev/null
+++ b/Engine/Plugins/Developer/Concert/ConcertUI/ConcertClientSharedSlate/Source/ConcertClientSharedSlate/Public/Replication/Editor/Model/IReplicationStreamModel.h
@@ -0,0 +1,107 @@
+// Copyright Epic Games, Inc. All Rights Reserved.
+
+#pragma once
+
+#include "Misc/EBreakBehavior.h"
+#include "Replication/Data/ConcertPropertySelection.h"
+#include "Templates/Function.h"
+#include "Templates/SharedPointer.h"
+
+namespace UE::ConcertClientSharedSlate
+{
+	/**
+	 * Abstracts the concept of mapping objects to properties. This is a read-only interface.
+	 * @see IEditableReplicationStreamModel 
+	 */
+	class CONCERTCLIENTSHAREDSLATE_API IReplicationStreamModel : public TSharedFromThis<IReplicationStreamModel>
+	{
+	public:
+
+		/** @return The number of replicated objects */
+		virtual uint32 GetNumReplicatedObjects() const = 0;
+		/** @return The number of properties assigned to this object */
+		virtual uint32 GetNumProperties(const FSoftObjectPath& Object) const = 0;
+
+		/** @return Gets the class of the replication object */
+		virtual FSoftClassPath GetObjectClass(const FSoftObjectPath& Object) const = 0;
+
+		/** @return Whether the objects are selected */
+		virtual bool ContainsObjects(const TSet<FSoftObjectPath>& Objects) const = 0;
+		/** @return Whether the properties are selected */
+		virtual bool ContainsProperties(const FSoftObjectPath& Object, const TSet<FConcertPropertyChain>& Properties) const = 0;
+		
+		/**
+		 * Iterates all objects in the mapping.
+		 * @return Whether there were any mapped objects
+		 */
+		virtual bool ForEachReplicatedObject(TFunctionRef<EBreakBehavior(const FSoftObjectPath& Object)> Delegate) const = 0;
+
+		/**
+		 * Iterates all the properties of the given object in the mapping.
+		 * @param Object The object of which to get the properties
+		 * @param Delegate The callback
+		 */
+		virtual bool ForEachProperty(const FSoftObjectPath& Object, TFunctionRef<EBreakBehavior(const FConcertPropertyChain& Property)> Delegate) const = 0;
+		
+		/**
+		 * Iterates all the direct child properties of a given property in the mapping.
+		 *
+		 * @param Object The object of which to get the properties
+		 * @param Delegate The callback
+		 * @param OptionalParent Defines of which properties to get the child properties. If the property chain is empty, gets the object's root properties.
+		 *
+		 * @return Whether any properties were mapped
+		 */
+		bool ForEachDirectChildProperty(const FSoftObjectPath& Object, TFunctionRef<EBreakBehavior(const FConcertPropertyChain& Property)> Delegate, const FConcertPropertyChain& OptionalParent = {}) const
+		{
+			bool bFoundAtLeastOne = false;
+			ForEachProperty(Object, [&Delegate, &OptionalParent, &bFoundAtLeastOne](const FConcertPropertyChain& Property) mutable
+			{
+				const bool bCanCall = Property == OptionalParent;
+				bFoundAtLeastOne |= bCanCall;
+				return bCanCall
+					? Delegate(Property)
+					: EBreakBehavior::Continue;
+			});
+			return bFoundAtLeastOne;
+		}
+
+		/** Util for getting properties as array */
+		TSet<FConcertPropertyChain> GetAllProperties(const FSoftObjectPath& Object) const
+		{
+			TSet<FConcertPropertyChain> Result;
+			ForEachProperty(Object, [&Result](const FConcertPropertyChain& Property)
+			{
+				Result.Add(Property);
+				return EBreakBehavior::Continue;
+			});
+			return Result;
+		}
+		
+		/** Util for getting child properties as array */
+		TSet<FName> GetChildProperties(const FSoftObjectPath& Object, const FConcertPropertyChain& OptionalParent = {}) const
+		{
+			TSet<FName> Result;
+			ForEachDirectChildProperty(Object, [&Result](const FConcertPropertyChain& Property)
+			{
+				Result.Add(Property.GetLeafProperty());
+				return EBreakBehavior::Continue;
+			}, OptionalParent);
+			return Result;
+		}
+
+		/** Iterates every subobject of Parent. */
+		void ForEachSubobject(const FSoftObjectPath& Parent, TFunctionRef<EBreakBehavior(const FSoftObjectPath& Child)> Callback) const
+		{
+			const FString ParentPathString = Parent.ToString();
+			ForEachReplicatedObject([&Parent, &ParentPathString, &Callback](const FSoftObjectPath& Object)
+			{
+				return Object != Parent && Object.ToString().Contains(ParentPathString)
+					? Callback(Object)
+					: EBreakBehavior::Continue;
+			});
+		}
+		
+		virtual ~IReplicationStreamModel() = default;
+	};
+}
diff --git a/Engine/Plugins/Developer/Concert/ConcertUI/ConcertClientSharedSlate/Source/ConcertClientSharedSlate/Public/Replication/Editor/Model/Object/IObjectSelectionSourceModel.h b/Engine/Plugins/Developer/Concert/ConcertUI/ConcertClientSharedSlate/Source/ConcertClientSharedSlate/Public/Replication/Editor/Model/Object/IObjectSelectionSourceModel.h
--- a/Engine/Plugins/Developer/Concert/ConcertUI/ConcertClientSharedSlate/Source/ConcertClientSharedSlate/Public/Replication/Editor/Model/Object/IObjectSelectionSourceModel.h
+++ b/Engine/Plugins/Developer/Concert/ConcertUI/ConcertClientSharedSlate/Source/ConcertClientSharedSlate/Public/Replication/Editor/Model/Object/IObjectSelectionSourceModel.h
@@ -21,7 +21,7 @@ namespace UE::ConcertClientSharedSlate
 		Invalid
 	};
 	
-	/** Decides which objects can be added to a IEditableObjectToPropertiesModel. */
+	/** Decides which objects can be added to a IEditableReplicationStreamModel. */
 	class IObjectSelectionSourceModel : public TSharedFromThis<IObjectSelectionSourceModel>
 	{
 	public:
diff --git a/Engine/Plugins/Developer/Concert/ConcertUI/ConcertClientSharedSlate/Source/ConcertClientSharedSlate/Public/Replication/Editor/Model/Property/IPropertySelectionSourceModel.h b/Engine/Plugins/Developer/Concert/ConcertUI/ConcertClientSharedSlate/Source/ConcertClientSharedSlate/Public/Replication/Editor/Model/Property/IPropertySelectionSourceModel.h
--- a/Engine/Plugins/Developer/Concert/ConcertUI/ConcertClientSharedSlate/Source/ConcertClientSharedSlate/Public/Replication/Editor/Model/Property/IPropertySelectionSourceModel.h
+++ b/Engine/Plugins/Developer/Concert/ConcertUI/ConcertClientSharedSlate/Source/ConcertClientSharedSlate/Public/Replication/Editor/Model/Property/IPropertySelectionSourceModel.h
@@ -10,7 +10,7 @@ namespace UE::ConcertClientSharedSlate
 {
 	using FPropertySourceCategory = ConcertSharedSlate::TSourceSelectionCategory<FSelectablePropertyInfo>;
 	
-	/** Decides which properties can be added to a IObjectToPropertiesModel. */
+	/** Decides which properties can be added to a IReplicationStreamModel. */
 	class IPropertySelectionSourceModel : public TSharedFromThis<IPropertySelectionSourceModel>
 	{
 	public:
diff --git a/Engine/Plugins/Developer/Concert/ConcertUI/ConcertClientSharedSlate/Source/ConcertClientSharedSlate/Public/Replication/Editor/Model/Property/SelectPropertyFromUClassModel.h b/Engine/Plugins/Developer/Concert/ConcertUI/ConcertClientSharedSlate/Source/ConcertClientSharedSlate/Public/Replication/Editor/Model/Property/SelectPropertyFromUClassModel.h
--- a/Engine/Plugins/Developer/Concert/ConcertUI/ConcertClientSharedSlate/Source/ConcertClientSharedSlate/Public/Replication/Editor/Model/Property/SelectPropertyFromUClassModel.h
+++ b/Engine/Plugins/Developer/Concert/ConcertUI/ConcertClientSharedSlate/Source/ConcertClientSharedSlate/Public/Replication/Editor/Model/Property/SelectPropertyFromUClassModel.h
@@ -9,7 +9,7 @@ namespace UE::ConcertClientSharedSlate
 	class FConcertSyncCoreReplicatedPropertySource;
 	
 	/**
-	 * Decides which properties can be added to IEditableObjectToPropertiesModel.
+	 * Decides which properties can be added to IEditableReplicationStreamModel.
 	 * The allowed properties are those returned by UE::ConcertSyncCore::ForEachReplicatableProperty.
 	 */
 	class CONCERTCLIENTSHAREDSLATE_API FSelectPropertyFromUClassModel : public IPropertySelectionSourceModel
diff --git a/Engine/Plugins/Developer/Concert/ConcertUI/ConcertClientSharedSlate/Source/ConcertClientSharedSlate/Public/Replication/Editor/View/ReplicationColumnsUtils.h b/Engine/Plugins/Developer/Concert/ConcertUI/ConcertClientSharedSlate/Source/ConcertClientSharedSlate/Public/Replication/Editor/View/ReplicationColumnsUtils.h
--- a/Engine/Plugins/Developer/Concert/ConcertUI/ConcertClientSharedSlate/Source/ConcertClientSharedSlate/Public/Replication/Editor/View/ReplicationColumnsUtils.h
+++ b/Engine/Plugins/Developer/Concert/ConcertUI/ConcertClientSharedSlate/Source/ConcertClientSharedSlate/Public/Replication/Editor/View/ReplicationColumnsUtils.h
@@ -11,9 +11,9 @@ struct FConcertPropertyChain;
 
 namespace UE::ConcertClientSharedSlate
 {
-	class IEditableObjectToPropertiesModel;
+	class IEditableReplicationStreamModel;
 	class IReplicationStreamViewer;
-	class IObjectToPropertiesModel;
+	class IReplicationStreamModel;
 }
 
 namespace UE::ConcertClientSharedSlate
diff --git a/Engine/Plugins/Developer/Concert/ConcertUI/ConcertClientSharedSlate/Source/ConcertClientSharedSlate/Public/Replication/ReplicationWidgetFactories.h b/Engine/Plugins/Developer/Concert/ConcertUI/ConcertClientSharedSlate/Source/ConcertClientSharedSlate/Public/Replication/ReplicationWidgetFactories.h
--- a/Engine/Plugins/Developer/Concert/ConcertUI/ConcertClientSharedSlate/Source/ConcertClientSharedSlate/Public/Replication/ReplicationWidgetFactories.h
+++ b/Engine/Plugins/Developer/Concert/ConcertUI/ConcertClientSharedSlate/Source/ConcertClientSharedSlate/Public/Replication/ReplicationWidgetFactories.h
@@ -16,8 +16,8 @@ struct FObjectReplicationMap;
 
 namespace UE::ConcertClientSharedSlate
 {
-	class IEditableObjectToPropertiesModel;
-	class IObjectToPropertiesModel;
+	class IEditableReplicationStreamModel;
+	class IReplicationStreamModel;
 	class IObjectSelectionSourceModel;
 	class IReplicationStreamEditor;
 	class IReplicationStreamViewer;
@@ -36,7 +36,7 @@ namespace UE::ConcertClientSharedSlate
 	 * 
 	 * @return A model that will edit the FObjectReplicationMap.
 	 */
-	CONCERTCLIENTSHAREDSLATE_API TSharedRef<IEditableObjectToPropertiesModel> CreatePropertySelectionModel(
+	CONCERTCLIENTSHAREDSLATE_API TSharedRef<IEditableReplicationStreamModel> CreatePropertySelectionModel(
 		UObject& OwnerObject,
 		TAttribute<FObjectReplicationMap*> ReplicationMapAttribute
 		);
@@ -64,7 +64,7 @@ namespace UE::ConcertClientSharedSlate
 		 * The model that the editor is displaying.
 		 * @note The view will keep a strong reference to this.
 		 */
-		TSharedRef<IEditableObjectToPropertiesModel> DataModel;
+		TSharedRef<IEditableReplicationStreamModel> DataModel;
 		/**
 		 * Determines the objects that can be added to the object list. 
 		 * @note The view will keep a strong reference to this.
