Remove subobject view

diff --git a/Engine/Plugins/Developer/Concert/ConcertApp/MultiUserClient/Source/MultiUserClient/Private/Widgets/ActiveSession/Replication/Client/Multi/Columns/MultiStreamColumns_ReplicationToggle.cpp b/Engine/Plugins/Developer/Concert/ConcertApp/MultiUserClient/Source/MultiUserClient/Private/Widgets/ActiveSession/Replication/Client/Multi/Columns/MultiStreamColumns_ReplicationToggle.cpp
--- a/Engine/Plugins/Developer/Concert/ConcertApp/MultiUserClient/Source/MultiUserClient/Private/Widgets/ActiveSession/Replication/Client/Multi/Columns/MultiStreamColumns_ReplicationToggle.cpp
+++ b/Engine/Plugins/Developer/Concert/ConcertApp/MultiUserClient/Source/MultiUserClient/Private/Widgets/ActiveSession/Replication/Client/Multi/Columns/MultiStreamColumns_ReplicationToggle.cpp
@@ -5,6 +5,7 @@
 #include "SReplicationMultiToggleCheckbox.h"
 
 #include "Widgets/Layout/SBox.h"
+#include "Widgets/Views/SHeaderRow.h"
 
 #define LOCTEXT_NAMESPACE "ReplicationToggle"
 
diff --git a/Engine/Plugins/Developer/Concert/ConcertApp/MultiUserClient/Source/MultiUserClient/Private/Widgets/ActiveSession/Replication/Client/Single/Columns/SingleClientColumns.h b/Engine/Plugins/Developer/Concert/ConcertApp/MultiUserClient/Source/MultiUserClient/Private/Widgets/ActiveSession/Replication/Client/Single/Columns/SingleClientColumns.h
--- a/Engine/Plugins/Developer/Concert/ConcertApp/MultiUserClient/Source/MultiUserClient/Private/Widgets/ActiveSession/Replication/Client/Single/Columns/SingleClientColumns.h
+++ b/Engine/Plugins/Developer/Concert/ConcertApp/MultiUserClient/Source/MultiUserClient/Private/Widgets/ActiveSession/Replication/Client/Single/Columns/SingleClientColumns.h
@@ -24,14 +24,8 @@ namespace UE::MultiUserClient
 
 namespace UE::MultiUserClient::SingleClientColumns
 {
-	/** @see UE::ConcertClientSharedSlate::ReplicationColumns::Subobject::ETopLevelColumnOrder */
+	/** @see UE::ConcertClientSharedSlate::ReplicationColumns::TopLevel::ETopLevelColumnOrder */
 	enum class ETopLevelObjectColumnOrder
-	{
-		ToggleAuthority = 0,
-		Owner = 40
-	};
-	/** @see UE::ConcertClientSharedSlate::ReplicationColumns::Subobject::ESubobjectColumnOrder */
-	enum class ESubobjectColumnOrder
 	{
 		ToggleAuthority = 0,
 		ConflictWarning = 5,
@@ -73,7 +67,7 @@ namespace UE::MultiUserClient::SingleClientColumns
 	 * @param AuthoritySynchronizer Used to determine whether changing authority is at all enabled
 	 * @return Column that can be placed in the table
 	 */
-	ConcertClientSharedSlate::ReplicationColumns::FReplicationSubobjectObjectColumn ToggleSubobjectAuthority(
+	ConcertClientSharedSlate::ReplicationColumns::FReplicationTopLevelObjectColumn ToggleObjectAuthority(
 		FAuthorityChangeTracker& ChangeTracker,
 		ISubmissionWorkflow& SubmissionWorkflow
 		);
@@ -102,7 +96,7 @@ namespace UE::MultiUserClient::SingleClientColumns
 	 * @param InAuthorityCache Used to determine which client owns the property
 	 * @return Column that can be placed in the table
 	 */
-	ConcertClientSharedSlate::ReplicationColumns::FReplicationSubobjectObjectColumn OwnerOfSubobject(
+	ConcertClientSharedSlate::ReplicationColumns::FReplicationTopLevelObjectColumn OwnerOfObject(
 		const TSharedRef<IConcertClient>& InClient,
 		FGlobalAuthorityCache& InAuthorityCache
 		);
@@ -135,7 +129,7 @@ namespace UE::MultiUserClient::SingleClientColumns
 	 * 
 	 * @return Column that can be placed in the table
 	 */
-	ConcertClientSharedSlate::ReplicationColumns::FReplicationSubobjectObjectColumn ConflictWarningForSubobject(
+	ConcertClientSharedSlate::ReplicationColumns::FReplicationTopLevelObjectColumn ConflictWarningForObject(
 		TSharedRef<IConcertClient> InClient,
 		FGlobalAuthorityCache& InAuthorityCache,
 		const FGuid& ClientId
diff --git a/Engine/Plugins/Developer/Concert/ConcertApp/MultiUserClient/Source/MultiUserClient/Private/Widgets/ActiveSession/Replication/Client/Single/Columns/SingleClientColumns_ConflictWarning.cpp b/Engine/Plugins/Developer/Concert/ConcertApp/MultiUserClient/Source/MultiUserClient/Private/Widgets/ActiveSession/Replication/Client/Single/Columns/SingleClientColumns_ConflictWarning.cpp
--- a/Engine/Plugins/Developer/Concert/ConcertApp/MultiUserClient/Source/MultiUserClient/Private/Widgets/ActiveSession/Replication/Client/Single/Columns/SingleClientColumns_ConflictWarning.cpp
+++ b/Engine/Plugins/Developer/Concert/ConcertApp/MultiUserClient/Source/MultiUserClient/Private/Widgets/ActiveSession/Replication/Client/Single/Columns/SingleClientColumns_ConflictWarning.cpp
@@ -21,13 +21,13 @@ namespace UE::MultiUserClient::SingleClientColumns
 	const FName ConflictWarningSubobjectColumnId = TEXT("ConflictWarningSubobjectColumn");
 	const FName ConflictWarningPropertyColumnId = TEXT("ConflictWarningPropertyColumn");
 	
-	ConcertClientSharedSlate::ReplicationColumns::FReplicationSubobjectObjectColumn ConflictWarningForSubobject(
+	ConcertClientSharedSlate::ReplicationColumns::FReplicationTopLevelObjectColumn ConflictWarningForObject(
 		TSharedRef<IConcertClient> InClient,
 		FGlobalAuthorityCache& InAuthorityCache,
 		const FGuid& ClientId
 		)
 	{
-		using FColumnType = ConcertClientSharedSlate::ReplicationColumns::FReplicationSubobjectObjectColumn;
+		using FColumnType = ConcertClientSharedSlate::ReplicationColumns::FReplicationTopLevelObjectColumn;
 		return FColumnType(
 			FColumnType::FArguments()
 				.GenerateWidgetColumn_Lambda([InClient = MoveTemp(InClient), &InAuthorityCache, ClientId](const FColumnType::FBuildArgs& InArgs) mutable
@@ -63,7 +63,7 @@ namespace UE::MultiUserClient::SingleClientColumns
 						.ToolTipText_Lambda(GetToolTip);
 				})
 				.PopulateSearchItems_Lambda([](const ConcertClientSharedSlate::FReplicatedObjectData&, TArray<FString>&){})
-				.ColumnSortOrder(static_cast<int32>(ESubobjectColumnOrder::ConflictWarning)),
+				.ColumnSortOrder(static_cast<int32>(ETopLevelObjectColumnOrder::ConflictWarning)),
 			SHeaderRow::Column(ConflictWarningSubobjectColumnId)
 				.DefaultLabel(FText::GetEmpty())
 				.FixedWidth(20.f)
diff --git a/Engine/Plugins/Developer/Concert/ConcertApp/MultiUserClient/Source/MultiUserClient/Private/Widgets/ActiveSession/Replication/Client/Single/Columns/SingleClientColumns_Owner.cpp b/Engine/Plugins/Developer/Concert/ConcertApp/MultiUserClient/Source/MultiUserClient/Private/Widgets/ActiveSession/Replication/Client/Single/Columns/SingleClientColumns_Owner.cpp
--- a/Engine/Plugins/Developer/Concert/ConcertApp/MultiUserClient/Source/MultiUserClient/Private/Widgets/ActiveSession/Replication/Client/Single/Columns/SingleClientColumns_Owner.cpp
+++ b/Engine/Plugins/Developer/Concert/ConcertApp/MultiUserClient/Source/MultiUserClient/Private/Widgets/ActiveSession/Replication/Client/Single/Columns/SingleClientColumns_Owner.cpp
@@ -57,7 +57,7 @@ namespace UE::MultiUserClient::SingleClientColumns
 		ConcertClientSharedSlate::IReplicationStreamModel& InObjectModel
 		)
 	{
-		using FColumnType = ConcertClientSharedSlate::ReplicationColumns::FReplicationSubobjectObjectColumn;
+		using FColumnType = ConcertClientSharedSlate::ReplicationColumns::FReplicationTopLevelObjectColumn;
 		return FColumnType(
 			FColumnType::FArguments()
 				.GenerateWidgetColumn_Lambda([InClient, &InObjectModel, &InAuthorityCache](const FColumnType::FBuildArgs& InArgs)
@@ -84,12 +84,12 @@ namespace UE::MultiUserClient::SingleClientColumns
 			);
 	}
 	
-	ConcertClientSharedSlate::ReplicationColumns::FReplicationSubobjectObjectColumn OwnerOfSubobject(
+	ConcertClientSharedSlate::ReplicationColumns::FReplicationTopLevelObjectColumn OwnerOfObject(
 		const TSharedRef<IConcertClient>& InClient,
 		FGlobalAuthorityCache& InAuthorityCache
 		)
 	{
-		using FColumnType = ConcertClientSharedSlate::ReplicationColumns::FReplicationSubobjectObjectColumn;
+		using FColumnType = ConcertClientSharedSlate::ReplicationColumns::FReplicationTopLevelObjectColumn;
 		return FColumnType(
 			FColumnType::FArguments()
 				.GenerateWidgetColumn_Lambda([InClient, &InAuthorityCache](const FColumnType::FBuildArgs& InArgs)
@@ -109,7 +109,7 @@ namespace UE::MultiUserClient::SingleClientColumns
 						InOutSearchStrings
 						);
 				})
-				.ColumnSortOrder(static_cast<int32>(ESubobjectColumnOrder::Owner)),
+				.ColumnSortOrder(static_cast<int32>(ETopLevelObjectColumnOrder::Owner)),
 			SHeaderRow::Column(OwnerOfSubobjectColumnId)
 				.DefaultLabel(LOCTEXT("Subobject.Owner", "Owner"))
 				.FillSized(Private::Shared::DefaultWidth)
diff --git a/Engine/Plugins/Developer/Concert/ConcertApp/MultiUserClient/Source/MultiUserClient/Private/Widgets/ActiveSession/Replication/Client/Single/Columns/SingleClientColumns_ToggleAuthority.cpp b/Engine/Plugins/Developer/Concert/ConcertApp/MultiUserClient/Source/MultiUserClient/Private/Widgets/ActiveSession/Replication/Client/Single/Columns/SingleClientColumns_ToggleAuthority.cpp
--- a/Engine/Plugins/Developer/Concert/ConcertApp/MultiUserClient/Source/MultiUserClient/Private/Widgets/ActiveSession/Replication/Client/Single/Columns/SingleClientColumns_ToggleAuthority.cpp
+++ b/Engine/Plugins/Developer/Concert/ConcertApp/MultiUserClient/Source/MultiUserClient/Private/Widgets/ActiveSession/Replication/Client/Single/Columns/SingleClientColumns_ToggleAuthority.cpp
@@ -163,7 +163,7 @@ namespace UE::MultiUserClient::SingleClientColumns
 		);
 	}
 
-	ConcertClientSharedSlate::ReplicationColumns::FReplicationSubobjectObjectColumn ToggleSubobjectAuthority(
+	ConcertClientSharedSlate::ReplicationColumns::FReplicationTopLevelObjectColumn ToggleObjectAuthority(
 		FAuthorityChangeTracker& ChangeTracker,
 		ISubmissionWorkflow& SubmissionWorkflow
 		)
@@ -210,7 +210,7 @@ namespace UE::MultiUserClient::SingleClientColumns
 					})
 				),
 			FText::GetEmpty(),
-			static_cast<int32>(ESubobjectColumnOrder::ToggleAuthority)
+			static_cast<int32>(ETopLevelObjectColumnOrder::ToggleAuthority)
 		);
 	}
 }
diff --git a/Engine/Plugins/Developer/Concert/ConcertApp/MultiUserClient/Source/MultiUserClient/Private/Widgets/ActiveSession/Replication/Client/Single/SReplicationClientView.cpp b/Engine/Plugins/Developer/Concert/ConcertApp/MultiUserClient/Source/MultiUserClient/Private/Widgets/ActiveSession/Replication/Client/Single/SReplicationClientView.cpp
--- a/Engine/Plugins/Developer/Concert/ConcertApp/MultiUserClient/Source/MultiUserClient/Private/Widgets/ActiveSession/Replication/Client/Single/SReplicationClientView.cpp
+++ b/Engine/Plugins/Developer/Concert/ConcertApp/MultiUserClient/Source/MultiUserClient/Private/Widgets/ActiveSession/Replication/Client/Single/SReplicationClientView.cpp
@@ -14,7 +14,6 @@
 #include "Widgets/ActiveSession/Replication/Client/Single/Columns/SingleClientColumns.h"
 #include "Widgets/ActiveSession/Replication/Client/SClientToolbar.h"
 
-#include "HAL/IConsoleManager.h"
 #include "Widgets/SBoxPanel.h"
 #include "Widgets/Layout/SBox.h"
 
@@ -22,12 +21,6 @@
 
 namespace UE::MultiUserClient
 {
-	TAutoConsoleVariable<int32> CVarReplicationClientViewMode(
-		TEXT("MultiUser.ReplicationEditorMode"),
-		1,
-		TEXT("Determines the look of the editor mode.\n0 - Three Sections\n1 - Two sections")
-		);
-	
 	void SReplicationClientView::Construct(const FArguments& InArgs, const TSharedRef<IConcertClient>& InClient, FReplicationClientManager& InClientManager)
 	{
 		ConcertClient = InClient;
@@ -57,95 +50,15 @@ namespace UE::MultiUserClient
 			+SVerticalBox::Slot()
 			.FillHeight(1.f)
 			[
-				CreateEditorContent()
+				CreateContent(*ReplicationClient)
 			]
 		];
 		
 		// Refresh UI if streams change externally, e.g. a remote client changed what they sent
 		ReplicationClient->OnModelChanged().AddSP(this, &SReplicationClientView::OnModelChanged);
-
-		CVarReplicationClientViewMode->OnChangedDelegate().AddSP(this, &SReplicationClientView::OnConsoleVariableChanged);
-	}
-
-	TSharedRef<SWidget> SReplicationClientView::CreateEditorContent()
-	{
-		Content = SNew(SBox);
-		RebuildContent();
-		return Content.ToSharedRef();
-	}
-
-	void SReplicationClientView::RebuildContent()
-	{
-		FReplicationClient* ReplicationClient = GetReplicationClientAttribute.Get();
-		check(ReplicationClient);
-
-		// Only one widget should be visible at the time so they do not interfere with each other.
-		switch (CVarReplicationClientViewMode.GetValueOnGameThread())
-		{
-		case 0:
-			Content->SetContent(CreateThreeSectionedContent(*ReplicationClient));
-			break;
-		case 1:
-			Content->SetContent(CreateTwoSectionedContent(*ReplicationClient));
-			break;
-
-		default:
-			Content->SetContent(SNullWidget::NullWidget);
-		}
-	}
-
-	TSharedRef<SWidget> SReplicationClientView::CreateThreeSectionedContent(FReplicationClient& InReplicationClient)
-	{
-		using namespace ConcertClientSharedSlate;
-		IReplicationStreamModel& PropertyModel = *InReplicationClient.GetClientEditModel();
-		FAuthorityChangeTracker& AuthorityTracker = InReplicationClient.GetAuthorityDiffer();
-		ISubmissionWorkflow& SubmissionWorkflow = InReplicationClient.GetSubmissionWorkflow();
-		FGlobalAuthorityCache& AuthorityCache = ClientManager->GetAuthorityCache();
-		
-		const TAttribute<const IReplicationStreamViewer*> GetReplicationViewerAttribute =
-			TAttribute<const IReplicationStreamViewer*>::CreateLambda([this](){ return EditorView_TwoSectioned.Get(); });
-		TAttribute<const FConcertReplicationEditorSettings*> ReplicationSettingsAttribute =
-			TAttribute<const FConcertReplicationEditorSettings*>::CreateLambda([](){ return &UMultiUserReplicationSettings::Get()->ReplicationEditorSettings; });
-
-		// Add checkboxes in front of top level and subobject rows for changing authority
-		const FCreateSubobjectViewParams SubobjectViewParams
-		{
-			.AdditionalColumns =
-			{
-				SingleClientColumns::ToggleSubobjectAuthority(AuthorityTracker, SubmissionWorkflow),
-				SingleClientColumns::ConflictWarningForSubobject(ConcertClient.ToSharedRef(), AuthorityCache, InReplicationClient.GetEndpointId()),
-				SingleClientColumns::OwnerOfSubobject(ConcertClient.ToSharedRef(), AuthorityCache)
-			}
-		};
-		const FCreateEditorParams ReplicationEditorCreationParams
-		{
-			.DataModel = InReplicationClient.GetClientEditModel(),
-			.ObjectSource = MakeShared<FActorSelectionSourceModel>(),
-			.PropertySource = MakeShared<FSelectPropertyFromUClassModel>(),
-			.IsEditingEnabled = TAttribute<bool>::CreateLambda([&SubmissionWorkflow](){ return SubmissionWorkflow.GetUploadability() != EChangeUploadability::NotImplemented; }),
-			.EditingDisabledToolTipText = LOCTEXT("Editing.NotImplemented", "Editing remote clients is not implemented. You can only edit the local client."),
-			.ReplicationSettingsAttribute = MoveTemp(ReplicationSettingsAttribute),
-			.ViewerParams =
-			{
-				.SubobjectView = CreateUnrealEditorSubobjectView(SubobjectViewParams),
-				.AdditionalObjectColumns =
-				{
-					SingleClientColumns::ToggleTopLevelAuthority(PropertyModel, AuthorityTracker, SubmissionWorkflow),
-					SingleClientColumns::OwnerOfTopLevelObject(ConcertClient.ToSharedRef(), AuthorityCache, PropertyModel)
-				},
-				.AdditionalPropertyColumns =
-				{
-					SingleClientColumns::OwnerOfProperty(ConcertClient.ToSharedRef(), AuthorityCache, GetReplicationViewerAttribute),
-					SingleClientColumns::ConflictWarningForProperty(ConcertClient.ToSharedRef(), GetReplicationViewerAttribute, AuthorityCache, InReplicationClient.GetEndpointId())
-				}
-			}
-		};
-
-		EditorView_ThreeSectioned = CreateDefaultStreamEditor(ReplicationEditorCreationParams);
-		return EditorView_ThreeSectioned.ToSharedRef();
 	}
 
-	TSharedRef<SWidget> SReplicationClientView::CreateTwoSectionedContent(FReplicationClient& InReplicationClient)
+	TSharedRef<SWidget> SReplicationClientView::CreateContent(FReplicationClient& InReplicationClient)
 	{
 		using namespace ConcertClientSharedSlate;
 		FAuthorityChangeTracker& AuthorityTracker = InReplicationClient.GetAuthorityDiffer();
@@ -171,10 +84,9 @@ namespace UE::MultiUserClient
 				.SubobjectModel = CreateDefaultComponentHierarchySubobjectModel(), // This makes actors have children in the top view
 				.AdditionalObjectColumns =
 				{
-					// TODO DP: When we decide which UI version to choose, remove subobject from these function names.
-					SingleClientColumns::ToggleSubobjectAuthority(AuthorityTracker, SubmissionWorkflow),
-					SingleClientColumns::ConflictWarningForSubobject(ConcertClient.ToSharedRef(), AuthorityCache, InReplicationClient.GetEndpointId()),
-					SingleClientColumns::OwnerOfSubobject(ConcertClient.ToSharedRef(), AuthorityCache)
+					SingleClientColumns::ToggleObjectAuthority(AuthorityTracker, SubmissionWorkflow),
+					SingleClientColumns::ConflictWarningForObject(ConcertClient.ToSharedRef(), AuthorityCache, InReplicationClient.GetEndpointId()),
+					SingleClientColumns::OwnerOfObject(ConcertClient.ToSharedRef(), AuthorityCache)
 				},
 				.AdditionalPropertyColumns =
 				{
@@ -201,11 +113,6 @@ namespace UE::MultiUserClient
 		}
 	}
 
-	void SReplicationClientView::OnConsoleVariableChanged(IConsoleVariable* ConsoleVariable)
-	{
-		RebuildContent();
-	}
-
 	void SReplicationClientView::EnumerateReplicatedObjects(TFunctionRef<void(const FSoftObjectPath&)> Consumer) const
 	{
 		const FReplicationClient* ReplicationClient = GetReplicationClientAttribute.Get();
diff --git a/Engine/Plugins/Developer/Concert/ConcertApp/MultiUserClient/Source/MultiUserClient/Private/Widgets/ActiveSession/Replication/Client/Single/SReplicationClientView.h b/Engine/Plugins/Developer/Concert/ConcertApp/MultiUserClient/Source/MultiUserClient/Private/Widgets/ActiveSession/Replication/Client/Single/SReplicationClientView.h
--- a/Engine/Plugins/Developer/Concert/ConcertApp/MultiUserClient/Source/MultiUserClient/Private/Widgets/ActiveSession/Replication/Client/Single/SReplicationClientView.h
+++ b/Engine/Plugins/Developer/Concert/ConcertApp/MultiUserClient/Source/MultiUserClient/Private/Widgets/ActiveSession/Replication/Client/Single/SReplicationClientView.h
@@ -6,7 +6,6 @@
 #include "Widgets/SCompoundWidget.h"
 
 class IConcertClient;
-class SBox;
 
 namespace UE::ConcertClientSharedSlate
 {
@@ -35,9 +34,6 @@ namespace UE::MultiUserClient
 		void Construct(const FArguments& InArgs, const TSharedRef<IConcertClient>& InClient, FReplicationClientManager& InClientManager);
 
 	private:
-
-		/** Holds the dynamic content, which changes based on CVarReplicationClientViewMode  */
-		TSharedPtr<SBox> Content;
 		
 		/** Used to rebuild Content */
 		TSharedPtr<IConcertClient> ConcertClient;
@@ -50,17 +46,11 @@ namespace UE::MultiUserClient
 		
 		/** The client to depict. Should always return true. If the client is destroyed, so should this widget be. */
 		TAttribute<FReplicationClient*> GetReplicationClientAttribute;
-
-
-		TSharedRef<SWidget> CreateEditorContent();
-		void RebuildContent();
-		TSharedRef<SWidget> CreateThreeSectionedContent(FReplicationClient& InReplicationClient);
-		TSharedRef<SWidget> CreateTwoSectionedContent(FReplicationClient& InReplicationClient);
+		
+		TSharedRef<SWidget> CreateContent(FReplicationClient& InReplicationClient);
 		
 		/** Called when any of the streams change. */
 		void OnModelChanged() const;
-		
-		void OnConsoleVariableChanged(IConsoleVariable* ConsoleVariable);
 
 		/** Lists all objects being replicated */
 		void EnumerateReplicatedObjects(TFunctionRef<void(const FSoftObjectPath&)> Consumer) const;
diff --git a/Engine/Plugins/Developer/Concert/ConcertUI/ConcertClientSharedSlate/Source/ConcertClientSharedSlate/Private/Replication/Editor/View/ObjectEditor/SBaseReplicationStreamEditor.cpp b/Engine/Plugins/Developer/Concert/ConcertUI/ConcertClientSharedSlate/Source/ConcertClientSharedSlate/Private/Replication/Editor/View/ObjectEditor/SBaseReplicationStreamEditor.cpp
--- a/Engine/Plugins/Developer/Concert/ConcertUI/ConcertClientSharedSlate/Source/ConcertClientSharedSlate/Private/Replication/Editor/View/ObjectEditor/SBaseReplicationStreamEditor.cpp
+++ b/Engine/Plugins/Developer/Concert/ConcertUI/ConcertClientSharedSlate/Source/ConcertClientSharedSlate/Private/Replication/Editor/View/ObjectEditor/SBaseReplicationStreamEditor.cpp
@@ -49,7 +49,6 @@ namespace UE::ConcertClientSharedSlate
 			SAssignNew(ReplicationViewer, SReplicationStreamViewer, PropertiesModelAdapter.ToSharedRef())
 				.AdditionalObjectColumns(InArgs._AdditionalObjectColumns)
 				.AdditionalPropertyColumns(InArgs._AdditionalPropertyColumns)
-				.SubobjectView(InArgs._SubobjectView)
 				.SubobjectModel(InArgs._SubobjectModel)
 				.OnDeleteObjects(this, &SBaseReplicationStreamEditor::OnDeleteObjects)
 				.OnObjectsContextMenuOpening(this, &SBaseReplicationStreamEditor::OnObjectsContextMenuOpening)
@@ -226,10 +225,6 @@ namespace UE::ConcertClientSharedSlate
 		});
 		
 		EditablePropertiesModel->RemoveObjects(ObjectAndChildren);
-
-		// Viewer has an internal cache of the selected object. Setting the displayed object to the cached one is a no-op.
-		// Clear the cache so it displays correctly if the user adds the object back straight after removing it.
-		ReplicationViewer->ClearSubobjectSelection();
 	}
 
 	TSharedPtr<SWidget> SBaseReplicationStreamEditor::OnObjectsContextMenuOpening()
diff --git a/Engine/Plugins/Developer/Concert/ConcertUI/ConcertClientSharedSlate/Source/ConcertClientSharedSlate/Private/Replication/Editor/View/ObjectEditor/SBaseReplicationStreamEditor.h b/Engine/Plugins/Developer/Concert/ConcertUI/ConcertClientSharedSlate/Source/ConcertClientSharedSlate/Private/Replication/Editor/View/ObjectEditor/SBaseReplicationStreamEditor.h
--- a/Engine/Plugins/Developer/Concert/ConcertUI/ConcertClientSharedSlate/Source/ConcertClientSharedSlate/Private/Replication/Editor/View/ObjectEditor/SBaseReplicationStreamEditor.h
+++ b/Engine/Plugins/Developer/Concert/ConcertUI/ConcertClientSharedSlate/Source/ConcertClientSharedSlate/Private/Replication/Editor/View/ObjectEditor/SBaseReplicationStreamEditor.h
@@ -56,8 +56,6 @@ namespace UE::ConcertClientSharedSlate
 			/** Additional columns to add to the property view */
 			SLATE_ARGUMENT(TArray<ReplicationColumns::FReplicationPropertyColumn>, AdditionalPropertyColumns)
 		
-			/** Optional. Placed between root object outliner and property editor. */
-			SLATE_ARGUMENT(TSharedPtr<IReplicationSubobjectView>, SubobjectView)
 			/** Optional. If set, this determines the children nested under the root objects. */
 			SLATE_ARGUMENT(TSharedPtr<ISubobjectModel>, SubobjectModel)
 
diff --git a/Engine/Plugins/Developer/Concert/ConcertUI/ConcertClientSharedSlate/Source/ConcertClientSharedSlate/Private/Replication/Editor/View/ObjectEditor/SDefaultReplicationStreamEditor.cpp b/Engine/Plugins/Developer/Concert/ConcertUI/ConcertClientSharedSlate/Source/ConcertClientSharedSlate/Private/Replication/Editor/View/ObjectEditor/SDefaultReplicationStreamEditor.cpp
--- a/Engine/Plugins/Developer/Concert/ConcertUI/ConcertClientSharedSlate/Source/ConcertClientSharedSlate/Private/Replication/Editor/View/ObjectEditor/SDefaultReplicationStreamEditor.cpp
+++ b/Engine/Plugins/Developer/Concert/ConcertUI/ConcertClientSharedSlate/Source/ConcertClientSharedSlate/Private/Replication/Editor/View/ObjectEditor/SDefaultReplicationStreamEditor.cpp
@@ -3,7 +3,7 @@
 #include "SDefaultReplicationStreamEditor.h"
 
 #include "Replication/Editor/Model/IEditableReplicationStreamModel.h"
-#include "Replication/Editor/View/ObjectViewer/Tree/SelectionViewerColumns.h"
+#include "Replication/Editor/View/Tree/SelectionViewerColumns.h"
 
 #include "SBaseReplicationStreamEditor.h"
 
@@ -36,7 +36,6 @@ namespace UE::ConcertClientSharedSlate
 			SAssignNew(WrappedEditor, SBaseReplicationStreamEditor, InPropertiesModel, InObjectSelectionSource, InPropertySelectionSource)
 				.AdditionalObjectColumns(InArgs._AdditionalObjectColumns)
 				.AdditionalPropertyColumns(PropertyColumns)
-				.SubobjectView(InArgs._SubobjectView)
 				.SubobjectModel(InArgs._SubobjectModel)
 				.OnExtendObjectsContextMenu(InArgs._OnExtendObjectsContextMenu)
 				.SortPropertyRowPredicate(this, &SDefaultReplicationStreamEditor::SortPropertiesPredicate)
diff --git a/Engine/Plugins/Developer/Concert/ConcertUI/ConcertClientSharedSlate/Source/ConcertClientSharedSlate/Private/Replication/Editor/View/ObjectViewer/Property/SPropertyTreeView.cpp b/Engine/Plugins/Developer/Concert/ConcertUI/ConcertClientSharedSlate/Source/ConcertClientSharedSlate/Private/Replication/Editor/View/ObjectViewer/Property/SPropertyTreeView.cpp
new file mode 100644
--- /dev/null
+++ b/Engine/Plugins/Developer/Concert/ConcertUI/ConcertClientSharedSlate/Source/ConcertClientSharedSlate/Private/Replication/Editor/View/ObjectViewer/Property/SPropertyTreeView.cpp
@@ -0,0 +1,197 @@
+// Copyright Epic Games, Inc. All Rights Reserved.
+
+#include "SPropertyTreeView.h"
+
+#include "Filters/SBasicFilterBar.h"
+#include "PropertyFilter_ByPropertyType.h"
+#include "PropertyFrontendFilter.h"
+#include "Replication/Editor/Model/ReplicatedPropertyData.h"
+#include "Replication/Editor/View/DisplayUtils.h"
+
+#include "UObject/UnrealType.h"
+
+#define LOCTEXT_NAMESPACE "SReplicatedPropertiesView"
+
+namespace UE::ConcertClientSharedSlate
+{
+	/** Exposes SetFrontendFilterActive so we can manually enable the default filters */
+	class SReplicationFilterBar : public SBasicFilterBar<TSharedPtr<FReplicatedPropertyData>>
+	{
+		using Super = SBasicFilterBar<TSharedPtr<FReplicatedPropertyData>>;
+	public:
+
+		SLATE_BEGIN_ARGS(SReplicationFilterBar)
+		{}
+			SLATE_EVENT(FOnFilterChanged, OnFilterChanged)
+		SLATE_END_ARGS()
+
+		void Construct(const FArguments& InArgs, TArray<SPropertyTreeView::FFilterRef> AllFilters)
+		{
+			Super::Construct(
+			Super::FArguments()
+				.FilterPillStyle(EFilterPillStyle::Basic)
+				.CustomFilters(MoveTemp(AllFilters))
+				.OnFilterChanged(InArgs._OnFilterChanged)
+				.UseSectionsForCategories(true)
+				);
+		}
+
+		// Expose from SBasicFilterBar
+		using Super::SetFrontendFilterActive;
+
+		void SetFilterVisuallyEnabled(const SPropertyTreeView::FFilterRef& Filter, bool bEnabled)
+		{
+			const TSharedRef<SFilter>* FilterWidget = Filters.FindByPredicate([&Filter](const TSharedRef<SFilter>& FilterWidget)
+			{
+				return FilterWidget->GetFrontendFilter() == Filter;
+			});
+			if (ensure(FilterWidget))
+			{
+				FilterWidget->Get().SetEnabled(bEnabled);
+			}
+		}
+	};
+	
+	void SPropertyTreeView::Construct(const FArguments& InArgs)
+	{
+		SelectedObjectsAttribute = InArgs._SelectedObjects;
+		
+		const FBuildFilterBarResult Filters = BuildFilterBar();
+		ChildSlot
+		[
+			SAssignNew(ReplicatedProperties, SReplicationTreeView<FReplicatedPropertyData>)
+				.RootItemsSource(InArgs._RootItemsSource)
+				.OnGetChildren(InArgs._OnGetChildren)
+				.Columns(InArgs._Columns)
+				.ExpandableColumnLabel(InArgs._ExpandableColumnLabel)
+				.SelectionMode(InArgs._SelectionMode)
+				.FilterItem(this, &SPropertyTreeView::PassesFilters)
+				.LeftOfSearchBar()
+				[
+					SNew(SHorizontalBox)
+
+					// The combo button for selecting the property filters
+					+SHorizontalBox::Slot()
+					.AutoWidth()
+					[
+						SBasicFilterBar<TSharedPtr<FReplicatedPropertyData>>::MakeAddFilterButton(FilterBar.ToSharedRef())
+					]
+
+					+SHorizontalBox::Slot()
+					.AutoWidth()
+					[
+						InArgs._LeftOfSearchBar.Widget
+					]
+				]
+				.RowBelowSearchBar() [ FilterBar.ToSharedRef() ]
+				.NoItemsContent() [ SNew(STextBlock).Text(this, &SPropertyTreeView::GetAllFilteredText) ]
+		];
+
+		// For better UX, hide certain properties by default (e.g. why would you want to replicate bools?)
+		// Do this AFTER initializing ReplicatedProperties because it triggers the OnItemsChanged callback.
+		for (const FFilterRef& Filter : Filters.DisabledByDefault)
+		{
+			FilterBar->SetFilterCheckState(Filter, ECheckBoxState::Checked);
+		}
+		// Our filters are inverse (FFilterBase::IsInverseFilter) so that means these disabled filters will be active now (so e.g. bool will be filtered out by default now)
+		FilterBar->DisableAllFilters();
+		// Show all the other filters as enabled (not greyed out: blue) - they will not run their logic since they are inverse.
+		for (const FFilterRef& Filter : Filters.EnabledByDefault)
+		{
+			// Makes it appear on the bar
+			FilterBar->SetFilterCheckState(Filter, ECheckBoxState::Checked);
+			// Functionally makes it affect the search
+			FilterBar->SetFrontendFilterActive(Filter, true);
+			// Visually makes the button blue (i.e. so it looks enabled)
+			FilterBar->SetFilterVisuallyEnabled(Filter, true);
+		}
+	}
+
+	void SPropertyTreeView::OnItemsChanged() const
+	{
+		ReplicatedProperties->OnItemsChanged();
+	}
+
+	SPropertyTreeView::FBuildFilterBarResult SPropertyTreeView::BuildFilterBar()
+	{
+		TSharedRef<FFilterCategory> GeneralCategory = MakeShared<FFilterCategory>(
+			LOCTEXT("CommonCategory.Name", "General"),
+			LOCTEXT("CommonCategory.ToolTip", "Exclude general properties.")
+		);
+		TSharedRef<FFilterCategory> NumericCategory = MakeShared<FFilterCategory>(
+			LOCTEXT("NumericCategory.Name", "Numeric"),
+			LOCTEXT("NumericCategory.ToolTip", "Excludes numeric properties.")
+		);
+		TSharedRef<FFilterCategory> ContainersCategory = MakeShared<FFilterCategory>(
+			LOCTEXT("ContainersCategory.Name", "Containers"),
+			LOCTEXT("ContainersCategory.ToolTip", "Exclude container properties.")
+		);
+
+		FBuildFilterBarResult Result;
+		using FFrontendFilter = TPropertyFrontendFilter<FPropertyFilter_ByPropertyType>;
+		Result.DisabledByDefault =
+		{
+			MakeShared<FFrontendFilter>(MoveTemp(GeneralCategory), LOCTEXT("Bool", "Boolean"), TSet<FFieldClass*>{ FBoolProperty::StaticClass() }),
+			MakeShared<FFrontendFilter>(MoveTemp(NumericCategory), LOCTEXT("Byte", "Byte"), TSet<FFieldClass*>{ FByteProperty::StaticClass() }),
+			MakeShared<FFrontendFilter>(MoveTemp(GeneralCategory), LOCTEXT("Name", "Name"), TSet<FFieldClass*>{ FNameProperty::StaticClass() }),
+			MakeShared<FFrontendFilter>(MoveTemp(GeneralCategory), LOCTEXT("String", "String"), TSet<FFieldClass*>{ FStrProperty::StaticClass() }),
+			MakeShared<FFrontendFilter>(MoveTemp(GeneralCategory), LOCTEXT("Enum", "Enum"), TSet<FFieldClass*>{ FEnumProperty::StaticClass() })
+		};
+		Result.EnabledByDefault =
+		{
+			MakeShared<FFrontendFilter>(MoveTemp(NumericCategory), LOCTEXT("Ints", "Integers"), LOCTEXT("Ints.Tooltip", "Includes integer types: uint16, uint32, uint64, int16, int32, int64"), TSet<FFieldClass*>{ FUInt16Property::StaticClass(), FUInt32Property::StaticClass(), FUInt64Property::StaticClass(), FInt16Property::StaticClass(), FIntProperty::StaticClass(), FInt64Property::StaticClass() }),
+			MakeShared<FFrontendFilter>(MoveTemp(NumericCategory), LOCTEXT("Float", "Float"), TSet<FFieldClass*>{ FFloatProperty::StaticClass() }),
+			MakeShared<FFrontendFilter>(MoveTemp(NumericCategory), LOCTEXT("Double", "Double"), TSet<FFieldClass*>{ FDoubleProperty::StaticClass() }),
+			MakeShared<FFrontendFilter>(MoveTemp(GeneralCategory), LOCTEXT("Struct", "Struct"), TSet<FFieldClass*>{ FStructProperty::StaticClass() }),
+			MakeShared<FFrontendFilter>(MoveTemp(ContainersCategory), LOCTEXT("Array", "Array"), TSet<FFieldClass*>{ FArrayProperty::StaticClass() }),
+			MakeShared<FFrontendFilter>(MoveTemp(ContainersCategory), LOCTEXT("Set", "Set"), TSet<FFieldClass*>{ FSetProperty::StaticClass() }),
+			MakeShared<FFrontendFilter>(MoveTemp(ContainersCategory), LOCTEXT("Map", "Map"), TSet<FFieldClass*>{ FMapProperty::StaticClass() }),
+		};
+		
+		TArray<FFilterRef> AllFilters =
+		{
+			MakeShared<FFrontendFilter>(MoveTemp(NumericCategory), LOCTEXT("UInt16", "UInt16"), TSet<FFieldClass*>{ FUInt16Property::StaticClass() }),
+			MakeShared<FFrontendFilter>(MoveTemp(NumericCategory), LOCTEXT("UInt32", "UInt32"), TSet<FFieldClass*>{ FUInt32Property::StaticClass() }),
+			MakeShared<FFrontendFilter>(MoveTemp(NumericCategory), LOCTEXT("UInt64", "UInt64"), TSet<FFieldClass*>{ FUInt64Property::StaticClass() }),
+			MakeShared<FFrontendFilter>(MoveTemp(NumericCategory), LOCTEXT("Int16", "Int16"), TSet<FFieldClass*>{ FInt16Property::StaticClass() }),
+			MakeShared<FFrontendFilter>(MoveTemp(NumericCategory), LOCTEXT("Int32", "Int32"), TSet<FFieldClass*>{ FIntProperty::StaticClass() }),
+			MakeShared<FFrontendFilter>(MoveTemp(NumericCategory), LOCTEXT("Int64", "Int64"), TSet<FFieldClass*>{ FInt64Property::StaticClass() }),
+			MakeShared<FFrontendFilter>(MoveTemp(GeneralCategory), LOCTEXT("SoftPtr", "Soft Ptr"), TSet<FFieldClass*>{ FSoftObjectProperty::StaticClass() }),
+		};
+		AllFilters.Append(Result.DisabledByDefault);
+		AllFilters.Append(Result.EnabledByDefault);
+		
+		FilterBar = SNew(SReplicationFilterBar, MoveTemp(AllFilters))
+			.OnFilterChanged(this, &SPropertyTreeView::OnItemsChanged);
+		
+		return Result;
+	}
+
+	bool SPropertyTreeView::PassesFilters(const TSharedPtr<FReplicatedPropertyData>& ReplicatedPropertyData) const
+	{
+		return FilterBar->GetAllActiveFilters()->PassesAllFilters(ReplicatedPropertyData);
+	}
+
+	FText SPropertyTreeView::GetAllFilteredText() const
+	{
+		const TArray<FSoftObjectPath> Objects = SelectedObjectsAttribute.Get();
+		if (Objects.IsEmpty())
+		{
+			return FText::GetEmpty();
+		}
+
+		TSet<FString> Names;
+		Algo::Transform(Objects, Names, [](const FSoftObjectPath& ObjectPath)
+		{
+			return DisplayUtils::GetObjectDisplayText(ObjectPath).ToString();
+		});
+		
+		const FText ObjectsText = FText::FromString(FString::Join(Names, TEXT(", ")));
+		return FText::Format(LOCTEXT("AllFitleredFmt", "All properties filtered for selected {0}|plural(one=object,other=objects): {1}."),
+			Objects.Num(),
+			ObjectsText
+		);
+	}
+}
+
+#undef LOCTEXT_NAMESPACE
\ No newline at end of file
diff --git a/Engine/Plugins/Developer/Concert/ConcertUI/ConcertClientSharedSlate/Source/ConcertClientSharedSlate/Private/Replication/Editor/View/ObjectViewer/Property/SPropertyTreeView.h b/Engine/Plugins/Developer/Concert/ConcertUI/ConcertClientSharedSlate/Source/ConcertClientSharedSlate/Private/Replication/Editor/View/ObjectViewer/Property/SPropertyTreeView.h
new file mode 100644
--- /dev/null
+++ b/Engine/Plugins/Developer/Concert/ConcertUI/ConcertClientSharedSlate/Source/ConcertClientSharedSlate/Private/Replication/Editor/View/ObjectViewer/Property/SPropertyTreeView.h
@@ -0,0 +1,81 @@
+// Copyright Epic Games, Inc. All Rights Reserved.
+
+#pragma once
+
+#include "Filters/SBasicFilterBar.h"
+#include "Widgets/DeclarativeSyntaxSupport.h"
+#include "Widgets/SCompoundWidget.h"
+#include "Replication/Editor/View/Tree/SReplicationTreeView.h"
+
+namespace UE::ConcertClientSharedSlate
+{
+	class FReplicatedObjectData;
+	class FReplicatedPropertyData;
+	class SReplicationFilterBar;
+	
+	/** Displays a searchable tree view of properties (SReplicationTreeView) and decorates it with a SBasicFilterBar. */
+	class SPropertyTreeView : public SCompoundWidget
+	{
+	public:
+		
+		using FFilterRef = TSharedRef<FFilterBase<TSharedPtr<FReplicatedPropertyData>>>;
+		
+		SLATE_BEGIN_ARGS(SPropertyTreeView)
+		{}
+			/*************** Arguments inherited by SReplicationTreeView ***************/
+		
+			/** The items to display */
+			SLATE_ARGUMENT(TArray<TSharedPtr<FReplicatedPropertyData>>*, RootItemsSource)
+		
+			/** Gets an items children for the tree view */
+			SLATE_EVENT(SReplicationTreeView<FReplicatedPropertyData>::FGetItemChildren, OnGetChildren)
+		
+			/** The columns this list should have */
+			SLATE_ARGUMENT(TArray<TReplicationColumn<FReplicatedPropertyData>>, Columns)
+			/** The name of the column that will have the SExpanderArrow for the tree view. */
+			SLATE_ARGUMENT(FName, ExpandableColumnLabel)
+		
+			/** How many items are to allowed to be selected */
+			SLATE_ARGUMENT(ESelectionMode::Type, SelectionMode)
+		
+			/** Optional widget to add to the left of the search bar. */
+			SLATE_NAMED_SLOT(FArguments, LeftOfSearchBar)
+
+			/*************** Own arguments ***************/
+			// Please add new arguments here in the future
+
+			/** Gets the objects being property edited. */
+			SLATE_ATTRIBUTE(TArray<FSoftObjectPath>, SelectedObjects)
+		SLATE_END_ARGS()
+
+		void Construct(const FArguments& InArgs);
+
+		/** Called when the items need to be refiltered due to the item source changing. */
+		void OnItemsChanged() const;
+
+	private:
+
+		/** The tree view displaying the replicated properties */
+		TSharedPtr<SReplicationTreeView<FReplicatedPropertyData>> ReplicatedProperties;
+		/** Displays the active filters*/
+		TSharedPtr<SReplicationFilterBar> FilterBar;
+
+		/** Used to tell the user that the selected object have all properties filtered out. */
+		TAttribute<TArray<FSoftObjectPath>> SelectedObjectsAttribute;
+
+		struct FBuildFilterBarResult
+		{
+			TArray<FFilterRef> EnabledByDefault;
+			TArray<FFilterRef> DisabledByDefault;
+		};
+		/** Build the filter bar and returns the filters that should be active by default. */
+		FBuildFilterBarResult BuildFilterBar();
+
+		/** Runs all filters through this item */
+		bool PassesFilters(const TSharedPtr<FReplicatedPropertyData>& ReplicatedPropertyData) const;
+
+		/** Gets the message to display when all properties are filtered out. */
+		FText GetAllFilteredText() const;
+	};
+}
+
diff --git a/Engine/Plugins/Developer/Concert/ConcertUI/ConcertClientSharedSlate/Source/ConcertClientSharedSlate/Private/Replication/Editor/View/ObjectViewer/Property/SReplicatedPropertiesView.cpp b/Engine/Plugins/Developer/Concert/ConcertUI/ConcertClientSharedSlate/Source/ConcertClientSharedSlate/Private/Replication/Editor/View/ObjectViewer/Property/SReplicatedPropertiesView.cpp
deleted file mode 100644
diff --git a/Engine/Plugins/Developer/Concert/ConcertUI/ConcertClientSharedSlate/Source/ConcertClientSharedSlate/Private/Replication/Editor/View/ObjectViewer/Property/SReplicatedPropertiesView.h b/Engine/Plugins/Developer/Concert/ConcertUI/ConcertClientSharedSlate/Source/ConcertClientSharedSlate/Private/Replication/Editor/View/ObjectViewer/Property/SReplicatedPropertiesView.h
deleted file mode 100644
diff --git a/Engine/Plugins/Developer/Concert/ConcertUI/ConcertClientSharedSlate/Source/ConcertClientSharedSlate/Private/Replication/Editor/View/ObjectViewer/SReplicatedPropertyView.cpp b/Engine/Plugins/Developer/Concert/ConcertUI/ConcertClientSharedSlate/Source/ConcertClientSharedSlate/Private/Replication/Editor/View/ObjectViewer/SReplicatedPropertyView.cpp
new file mode 100644
--- /dev/null
+++ b/Engine/Plugins/Developer/Concert/ConcertUI/ConcertClientSharedSlate/Source/ConcertClientSharedSlate/Private/Replication/Editor/View/ObjectViewer/SReplicatedPropertyView.cpp
@@ -0,0 +1,207 @@
+// Copyright Epic Games, Inc. All Rights Reserved.
+
+#include "SReplicatedPropertyView.h"
+
+#include "Replication/Editor/Model/IReplicationStreamModel.h"
+#include "Replication/Editor/Model/ReplicatedObjectData.h"
+#include "Replication/Editor/Model/ReplicatedPropertyData.h"
+#include "Replication/Editor/View/ObjectViewer/Property/SPropertyTreeView.h"
+#include "Replication/Editor/View/Tree/SelectionViewerColumns.h"
+
+#include "Algo/AllOf.h"
+#include "Algo/ForEach.h"
+#include "Widgets/Layout/SBorder.h"
+#include "Widgets/Layout/SSplitter.h"
+#include "Widgets/Layout/SWidgetSwitcher.h"
+#include "Widgets/SBoxPanel.h"
+#include "Widgets/Text/STextBlock.h"
+
+#define LOCTEXT_NAMESPACE "SReplicatedPropertyView"
+
+namespace UE::ConcertClientSharedSlate
+{
+	void SReplicatedPropertyView::Construct(const FArguments& InArgs, TSharedRef<IReplicationStreamModel> InPropertiesModel)
+	{
+		PropertiesModel = MoveTemp(InPropertiesModel);
+		SortPropertyRowPredicate = InArgs._SortPropertyRowPredicate;
+		GetSelectedRootObjectsDelegate = InArgs._GetSelectedRootObjects;
+		check(GetSelectedRootObjectsDelegate.IsBound());
+		
+		ChildSlot
+		[
+			CreatePropertiesView(InArgs)
+		];
+	}
+	
+	void SReplicatedPropertyView::RefreshPropertyData()
+	{
+		const TArray<FSoftObjectPath> SelectedObjects = GetObjectsSelectedForPropertyEditing();
+		if (SelectedObjects.IsEmpty())
+		{
+			SetPropertyContent(EReplicatedPropertyContent::NoSelection);
+			return;
+		}
+
+		// Technically, the classes just need to be compatible with each other... but it is easier to just allow the same class.
+		const TOptional<FSoftClassPath> SharedClass = GetClassForPropertiesFromSelection(SelectedObjects);
+		if (!SharedClass)
+		{
+			SetPropertyContent(EReplicatedPropertyContent::SelectionTooBig);
+			return;
+		}
+		const FSoftClassPath Class = *SharedClass;
+		if (!ensureMsgf(Class.IsValid(), TEXT("This should only trigger if IReplicationSubobjectView returned a object not contained the model & that the UI is view-only. In that case fix your IReplicationSubobjectView or use it in an UI editor.")))
+		{
+			return;
+		}
+
+		// Build the set of properties that are shared by all of the selected objects
+		TSet<FConcertPropertyChain> SharedProperties = PropertiesModel->GetAllProperties(SelectedObjects[0]);
+		for (int32 i = 1; i < SelectedObjects.Num(); ++i)
+		{
+			SharedProperties = PropertiesModel->GetAllProperties(SelectedObjects[i]).Union(SharedProperties);
+		}
+		
+		// Try to re-use old instances by using the old ChainToPropertyDataCache. This is also done so the expansion states restore correctly in the tree view.
+		TMap<FConcertPropertyChain, TSharedPtr<FReplicatedPropertyData>> NewChainToPropertyDataCache;
+		
+		PropertyRowData.Empty();
+		for (const FConcertPropertyChain& PropertyChain : SharedProperties)
+		{
+			const TSharedPtr<FReplicatedPropertyData>* ExistingItem = ChainToPropertyDataCache.Find(PropertyChain);
+			const TSharedRef<FReplicatedPropertyData> Item = ExistingItem ? ExistingItem->ToSharedRef() : AllocatePropertyData(Class, PropertyChain);
+			PropertyRowData.Emplace(Item);
+			NewChainToPropertyDataCache.Emplace(PropertyChain, Item);
+		}
+
+		// If an item was removed, then NewPathToPropertyDataCache does not contain it. 
+		ChainToPropertyDataCache = MoveTemp(NewChainToPropertyDataCache);
+		
+		// The tree view requires the item source to only contain the root items.
+		BuildRootPropertyRowData();
+		SetPropertyContent(EReplicatedPropertyContent::Properties);
+		ReplicatedProperties->OnItemsChanged();
+	}
+
+	TArray<FSoftObjectPath> SReplicatedPropertyView::GetObjectsSelectedForPropertyEditing() const
+	{
+		TArray<FSoftObjectPath> Result;
+		Algo::Transform(GetSelectedRootObjectsDelegate.Execute(), Result, [](const TSharedPtr<FReplicatedObjectData>& ObjectData)
+		{
+			return ObjectData->GetObjectPath();
+		});
+		return Result;
+	}
+
+	TSharedRef<SWidget> SReplicatedPropertyView::CreatePropertiesView(const FArguments& InArgs)
+	{
+		TArray Columns
+		{
+			ReplicationColumns::Property::LabelColumn(),
+			ReplicationColumns::Property::TypeColumn()
+		};
+		Columns.Append(InArgs._AdditionalPropertyColumns);
+		
+		return SAssignNew(PropertyContent, SWidgetSwitcher)
+			// Make sure the slots are coherent with the order of EReplicatedPropertyContent!
+			.WidgetIndex(static_cast<int32>(EReplicatedPropertyContent::NoSelection))
+			
+			// EReplicatedPropertyContent::Properties
+			+SWidgetSwitcher::Slot()
+			[
+				SAssignNew(ReplicatedProperties, SPropertyTreeView)
+				.RootItemsSource(&RootPropertyRowData)
+				.OnGetChildren(this, &SReplicatedPropertyView::GetPropertyRowChildren)
+				.Columns(Columns)
+				.ExpandableColumnLabel(ReplicationColumns::Property::LabelColumnId)
+				.SelectionMode(ESelectionMode::Multi)
+				.LeftOfSearchBar()
+				[
+					InArgs._LeftOfPropertySearchBar.Widget
+				]
+				.SelectedObjects(this, &SReplicatedPropertyView::GetObjectsSelectedForPropertyEditing)
+			]
+			
+			// EReplicatedPropertyContent::NoSelection
+			+SWidgetSwitcher::Slot()
+			.VAlign(VAlign_Center)
+			.HAlign(HAlign_Center)
+			[
+				SNew(STextBlock)
+				.Text(LOCTEXT("NoPropertyEditedObjects", "Select an object to see selected properties"))
+			]
+			
+			// EReplicatedPropertyContent::SelectionTooBig
+			+SWidgetSwitcher::Slot()
+			.VAlign(VAlign_Center)
+			.HAlign(HAlign_Center)
+			[
+				SNew(STextBlock)
+				.Text(LOCTEXT("SelectionTooBig", "Select objects of the same type type to see selected properties"))
+			];
+	}
+	
+	TSharedRef<FReplicatedPropertyData> SReplicatedPropertyView::AllocatePropertyData(FSoftClassPath OwningClass, FConcertPropertyChain PropertyChain)
+	{
+		return MakeShared<FReplicatedPropertyData>(MoveTemp(OwningClass), MoveTemp(PropertyChain));
+	}
+
+	void SReplicatedPropertyView::BuildRootPropertyRowData()
+	{
+		RootPropertyRowData.Empty(PropertyRowData.Num());
+		for (const TSharedPtr<FReplicatedPropertyData>& PropertyData : PropertyRowData)
+		{
+			if (PropertyData->GetProperty().IsRootProperty())
+			{
+				RootPropertyRowData.Emplace(PropertyData);
+			}
+		}
+
+		SortPropertyRowArray(RootPropertyRowData);
+	}
+	
+	TOptional<FSoftClassPath> SReplicatedPropertyView::GetClassForPropertiesFromSelection(const TArray<FSoftObjectPath>& Objects) const
+	{
+		const FSoftClassPath Class = PropertiesModel->GetObjectClass(Objects[0]);
+		const bool bAllHaveSameClass = Algo::AllOf(Objects, [this, Class](const FSoftObjectPath& Object)
+		{
+			return PropertiesModel->GetObjectClass(Object) == Class;
+		});
+		return bAllHaveSameClass ? Class : TOptional<FSoftClassPath>{};
+	}
+
+	void SReplicatedPropertyView::GetPropertyRowChildren(TSharedPtr<FReplicatedPropertyData> ReplicatedPropertyData, TFunctionRef<void(TSharedPtr<FReplicatedPropertyData>)> ProcessChild)
+	{
+		TArray<TSharedPtr<FReplicatedPropertyData>> Children;
+		
+		// Not the most efficient but it should be fine.
+		for (const TSharedPtr<FReplicatedPropertyData>& Data : PropertyRowData)
+		{
+			if (Data->GetProperty().IsDirectChildOf(ReplicatedPropertyData->GetProperty()))
+			{
+				Children.Add(Data);
+			}
+		}
+
+		SortPropertyRowArray(Children);
+		Algo::ForEach(Children, [&ProcessChild](const TSharedPtr<FReplicatedPropertyData>& Data){ ProcessChild(Data); });
+	}
+
+	void SReplicatedPropertyView::SortPropertyRowArray(TArray<TSharedPtr<FReplicatedPropertyData>>& ToSort) const
+	{
+		if (SortPropertyRowPredicate.IsBound())
+		{
+			ToSort.Sort([this](const TSharedPtr<FReplicatedPropertyData>& Left, const TSharedPtr<FReplicatedPropertyData>& Right)
+			{
+				return SortPropertyRowPredicate.Execute(*Left.Get(), *Right.Get());
+			});
+		}
+	}
+
+	void SReplicatedPropertyView::SetPropertyContent(EReplicatedPropertyContent Content) const
+	{
+		PropertyContent->SetActiveWidgetIndex(static_cast<int32>(Content));
+	}
+}
+
+#undef LOCTEXT_NAMESPACE
\ No newline at end of file
diff --git a/Engine/Plugins/Developer/Concert/ConcertUI/ConcertClientSharedSlate/Source/ConcertClientSharedSlate/Private/Replication/Editor/View/ObjectViewer/SReplicatedPropertyView.h b/Engine/Plugins/Developer/Concert/ConcertUI/ConcertClientSharedSlate/Source/ConcertClientSharedSlate/Private/Replication/Editor/View/ObjectViewer/SReplicatedPropertyView.h
new file mode 100644
--- /dev/null
+++ b/Engine/Plugins/Developer/Concert/ConcertUI/ConcertClientSharedSlate/Source/ConcertClientSharedSlate/Private/Replication/Editor/View/ObjectViewer/SReplicatedPropertyView.h
@@ -0,0 +1,100 @@
+// Copyright Epic Games, Inc. All Rights Reserved.
+
+#pragma once
+
+#include "Replication/Editor/View/Tree/SelectionViewerColumns.h"
+#include "Replication/ReplicationWidgetDelegates.h"
+
+#include "Misc/Optional.h"
+#include "Replication/Data/ConcertPropertySelection.h"
+#include "Widgets/DeclarativeSyntaxSupport.h"
+#include "Widgets/SCompoundWidget.h"
+
+class SWidgetSwitcher;
+
+namespace UE::ConcertClientSharedSlate
+{
+	class SPropertyTreeView;
+
+	/** Displays the SPropertyTreeView and decorates it with messages that prompt the user for action, e.g. to select an object to view properties. */
+	class SReplicatedPropertyView : public SCompoundWidget
+	{
+	public:
+		
+		DECLARE_DELEGATE_RetVal(TArray<TSharedPtr<FReplicatedObjectData>>, FGetSelectedRootObjects)
+
+		SLATE_BEGIN_ARGS(SReplicatedPropertyView)
+		{}
+			/** Additional columns to add to the property view */
+			SLATE_ARGUMENT(TArray<ReplicationColumns::FReplicationPropertyColumn>, AdditionalPropertyColumns)
+		
+			/** Optional. Used for determining the order in which properties are displayed. */
+			SLATE_EVENT(FSortPropertyPredicate, SortPropertyRowPredicate)
+			/** Gets the root objects selected in the object outliner. */
+			SLATE_EVENT(FGetSelectedRootObjects, GetSelectedRootObjects)
+		
+			/** Optional widget to add to the left of the property list search bar. */
+			SLATE_NAMED_SLOT(FArguments, LeftOfPropertySearchBar)
+		SLATE_END_ARGS()
+
+		void Construct(const FArguments& InArgs, TSharedRef<IReplicationStreamModel> InPropertiesModel);
+		
+		void RefreshPropertyData();
+		
+		const TArray<TSharedPtr<FReplicatedPropertyData>>& GetPropertyRowData() const { return PropertyRowData; }
+		TArray<FSoftObjectPath> GetObjectsSelectedForPropertyEditing() const;
+
+	private:
+		
+		/** The model this view is visualizing. */
+		TSharedPtr<IReplicationStreamModel> PropertiesModel;
+		
+		/** Tree view for replicated properties. Content depends on the current object selected. */
+		TSharedPtr<SPropertyTreeView> ReplicatedProperties;
+		
+		enum class EReplicatedPropertyContent
+		{
+			/** Shows the properties */
+			Properties,
+			/** Prompts: "Select an object to see selected properties" */
+			NoSelection,
+			/** Prompts: "Select objects of the same type type to see selected properties" */
+			SelectionTooBig
+		};
+		/** Determines the content displayed for PropertyArea. */
+		TSharedPtr<SWidgetSwitcher> PropertyContent;
+		
+		/** Used for determining the order in which properties are displayed. Default: Sort by name. */
+		FSortPropertyPredicate SortPropertyRowPredicate;
+		/** Gets the root objects selected in the object outliner. */
+		FGetSelectedRootObjects GetSelectedRootObjectsDelegate;
+		
+		/**
+		 * These instances can be subclasses of FReplicatedPropertyData, e.g. FReplicatedPropertyData_Editor.
+		 * Their type can be overridden by subclasses.
+		 * They only have the FReplicatedPropertyData type so they can be passed efficiently to SObjectToPropertyView.
+		 * @see GetPropertyData
+		 */
+		TArray<TSharedPtr<FReplicatedPropertyData>> PropertyRowData;
+		/** The instances of ObjectRowData which do not have any parents. This acts as the item source for the tree view. */
+		TArray<TSharedPtr<FReplicatedPropertyData>> RootPropertyRowData;
+		/** Inverse map of PropertyRowData using FReplicatedPropertyData::GetProperty as key. Contains all elements of PropertyRowData. */
+		TMap<FConcertPropertyChain, TSharedPtr<FReplicatedPropertyData>> ChainToPropertyDataCache;
+
+		TSharedRef<SWidget> CreatePropertiesView(const FArguments& InArgs);
+		
+		TSharedRef<FReplicatedPropertyData> AllocatePropertyData(FSoftClassPath OwningClass, FConcertPropertyChain PropertyChain);
+
+		/** Inits RootPropertyRowData from PropertyRowData. */
+		void BuildRootPropertyRowData();
+		/** Given the selected objects, determines whether they all have the same class and returns it if so. */
+		TOptional<FSoftClassPath> GetClassForPropertiesFromSelection(const TArray<FSoftObjectPath>& Objects) const;
+		void GetPropertyRowChildren(TSharedPtr<FReplicatedPropertyData> ReplicatedPropertyData, TFunctionRef<void(TSharedPtr<FReplicatedPropertyData>)> ProcessChild);
+		
+		// Utils
+		void SortPropertyRowArray(TArray<TSharedPtr<FReplicatedPropertyData>>& ToSort) const;
+		void SetPropertyContent(EReplicatedPropertyContent Content) const;
+	};
+}
+
+
diff --git a/Engine/Plugins/Developer/Concert/ConcertUI/ConcertClientSharedSlate/Source/ConcertClientSharedSlate/Private/Replication/Editor/View/ObjectViewer/SReplicationStreamViewer.cpp b/Engine/Plugins/Developer/Concert/ConcertUI/ConcertClientSharedSlate/Source/ConcertClientSharedSlate/Private/Replication/Editor/View/ObjectViewer/SReplicationStreamViewer.cpp
--- a/Engine/Plugins/Developer/Concert/ConcertUI/ConcertClientSharedSlate/Source/ConcertClientSharedSlate/Private/Replication/Editor/View/ObjectViewer/SReplicationStreamViewer.cpp
+++ b/Engine/Plugins/Developer/Concert/ConcertUI/ConcertClientSharedSlate/Source/ConcertClientSharedSlate/Private/Replication/Editor/View/ObjectViewer/SReplicationStreamViewer.cpp
@@ -6,9 +6,9 @@
 #include "Replication/Editor/Model/IReplicationStreamModel.h"
 #include "Replication/Editor/Model/ISubobjectModel.h"
 #include "Replication/Editor/Model/ReplicatedObjectData.h"
-#include "Replication/Editor/View/ObjectViewer/Property/SReplicatedPropertiesView.h"
-#include "Replication/Editor/View/ObjectViewer/Tree/SelectionViewerColumns.h"
-#include "SSubobjectAndPropertySection.h"
+#include "Replication/Editor/View/ObjectViewer/Property/SPropertyTreeView.h"
+#include "Replication/Editor/View/Tree/SelectionViewerColumns.h"
+#include "SReplicatedPropertyView.h"
 #include "Replication/Editor/View/ObjectUtils.h"
 
 #include "Widgets/Layout/SBorder.h"
@@ -38,7 +38,6 @@ namespace UE::ConcertClientSharedSlate
 	{
 		RefreshObjectData();
 		RefreshPropertyData();
-		RefreshSubobjectData();
 	}
 
 	TArray<FSoftObjectPath> SReplicationStreamViewer::GetSelectedTopLevelObjects() const
@@ -92,14 +91,6 @@ namespace UE::ConcertClientSharedSlate
 		}
 	}
 
-	void SReplicationStreamViewer::RefreshSubobjectData()
-	{
-		SubobjectAndPropertySection->RefreshSubobjectData();
-		
-		// Need to update properties in case the subobject view changed selection
-		RefreshPropertyData();
-	}
-
 	void SReplicationStreamViewer::RefreshPropertyData()
 	{
 		SubobjectAndPropertySection->RefreshPropertyData();
@@ -154,12 +145,7 @@ namespace UE::ConcertClientSharedSlate
 			ReplicatedObjects->SetSelectedItems(ItemsToExpand, true);
 		}
 	}
-
-	void SReplicationStreamViewer::ClearSubobjectSelection()
-	{
-		SubobjectAndPropertySection->ClearSubobjectSelection();
-	}
-
+	
 	TSharedRef<FReplicatedObjectData> SReplicationStreamViewer::AllocateObjectData(FSoftObjectPath ObjectPath)
 	{
 		return MakeShared<FReplicatedObjectData>(MoveTemp(ObjectPath));
@@ -200,12 +186,9 @@ namespace UE::ConcertClientSharedSlate
 			.RootItemsSource(&RootObjectRowData)
 			.OnGetChildren(this, &SReplicationStreamViewer::GetObjectRowChildren)
 			.OnContextMenuOpening(InArgs._OnObjectsContextMenuOpening)
-			.OnDeleteItems(InArgs._OnDeleteObjects)
-			.OnSelectionChanged_Lambda([this]()
+			.OnDeleteItems(InArgs._OnDeleteObjects).OnSelectionChanged_Lambda([this]()
 			{
-				RefreshSubobjectData();
-				// The subobject UI should automatically highlight the root objects after they are selected in the outliner
-				SubobjectAndPropertySection->SelectRootObjects();
+				RefreshPropertyData();
 			})
 			.Columns(Columns)
 			.ExpandableColumnLabel(ReplicationColumns::TopLevel::LabelColumnId)
@@ -237,9 +220,8 @@ namespace UE::ConcertClientSharedSlate
 				]
 				.BodyContent()
 				[
-					SAssignNew(SubobjectAndPropertySection, SSubobjectAndPropertySection, PropertiesModel.ToSharedRef())
+					SAssignNew(SubobjectAndPropertySection, SReplicatedPropertyView, PropertiesModel.ToSharedRef())
 					.AdditionalPropertyColumns(InArgs._AdditionalPropertyColumns)
-					.SubobjectView(InArgs._SubobjectView)
 					.SortPropertyRowPredicate(InArgs._SortPropertyRowPredicate)
 					.GetSelectedRootObjects_Lambda([this](){ return GetSelectedOutlinerObjects(); })
 					.LeftOfPropertySearchBar()
diff --git a/Engine/Plugins/Developer/Concert/ConcertUI/ConcertClientSharedSlate/Source/ConcertClientSharedSlate/Private/Replication/Editor/View/ObjectViewer/SReplicationStreamViewer.h b/Engine/Plugins/Developer/Concert/ConcertUI/ConcertClientSharedSlate/Source/ConcertClientSharedSlate/Private/Replication/Editor/View/ObjectViewer/SReplicationStreamViewer.h
--- a/Engine/Plugins/Developer/Concert/ConcertUI/ConcertClientSharedSlate/Source/ConcertClientSharedSlate/Private/Replication/Editor/View/ObjectViewer/SReplicationStreamViewer.h
+++ b/Engine/Plugins/Developer/Concert/ConcertUI/ConcertClientSharedSlate/Source/ConcertClientSharedSlate/Private/Replication/Editor/View/ObjectViewer/SReplicationStreamViewer.h
@@ -5,9 +5,9 @@
 #include "Replication/Editor/View/IReplicationStreamViewer.h"
 
 #include "Replication/Editor/Model/ReplicatedObjectData.h"
-#include "Replication/Editor/View/ObjectViewer/Tree/SReplicationTreeView.h"
-#include "Replication/Editor/View/ObjectViewer/Tree/SelectionViewerColumns.h"
-#include "SSubobjectAndPropertySection.h"
+#include "Replication/Editor/View/Tree/SReplicationTreeView.h"
+#include "Replication/Editor/View/Tree/SelectionViewerColumns.h"
+#include "SReplicatedPropertyView.h"
 
 #include "Algo/Transform.h"
 #include "Misc/Optional.h"
@@ -21,14 +21,13 @@ struct FSoftObjectPath;
 
 namespace UE::ConcertClientSharedSlate
 {
-	class SSubobjectAndPropertySection;
+	class SReplicatedPropertyView;
 	class FReplicatedPropertyData;
 	class FReplicatedObjectData;
 	class IEditableReplicationStreamModel;
-	class IReplicationSubobjectView;
 	class IReplicationStreamModel;
 	class ISubobjectModel;
-	class SReplicatedPropertiesView;
+	class SPropertyTreeView;
 	
 	/**
 	 * Root widget for viewing UMultiUserPropertyReplicationSelection.
@@ -51,8 +50,6 @@ namespace UE::ConcertClientSharedSlate
 			/** Additional columns to add to the property view */
 			SLATE_ARGUMENT(TArray<ReplicationColumns::FReplicationPropertyColumn>, AdditionalPropertyColumns)
 
-			/** Optional. Placed between root object outliner and property editor. */
-			SLATE_ARGUMENT(TSharedPtr<IReplicationSubobjectView>, SubobjectView)
 			/** Optional. If set, this determines the children nested under the root objects. */
 			SLATE_ARGUMENT(TSharedPtr<ISubobjectModel>, SubobjectModel)
 
@@ -83,7 +80,6 @@ namespace UE::ConcertClientSharedSlate
 		//~ End IReplicationStreamViewer Interface
 
 		void RefreshObjectData();
-		void RefreshSubobjectData();
 		void RefreshPropertyData();
 
 		/** Selects the given objects from the top level view, if applicable. */
@@ -91,9 +87,6 @@ namespace UE::ConcertClientSharedSlate
 
 		/** Expands the given objects, recursively if desired. */
 		void ExpandObjects(TConstArrayView<FSoftObjectPath> Objects, bool bRecursive);
-		
-		/** Clears all objects selected in the subobject view, if there is one. */
-		void ClearSubobjectSelection();
 
 		/** @return Gets the root objects selected in the outliner; the subobject view chooses which of these objects (or their subobjects) end up in GetSelectedObjectShowingProperties. */
 		TArray<TSharedPtr<FReplicatedObjectData>> GetSelectedOutlinerObjects() const { return ReplicatedObjects->GetSelectedItems(); }
@@ -108,7 +101,7 @@ namespace UE::ConcertClientSharedSlate
 		/** Lists the properties of the selected actor */
 		TSharedPtr<SExpandableArea> PropertyArea;
 		/** Edits the property list and (optionally) exposes subobjects of the selected root object. */
-		TSharedPtr<SSubobjectAndPropertySection> SubobjectAndPropertySection;
+		TSharedPtr<SReplicatedPropertyView> SubobjectAndPropertySection;
 
 		/** Tree view for replicated objects. */
 		TSharedPtr<SReplicationTreeView<FReplicatedObjectData>> ReplicatedObjects;
diff --git a/Engine/Plugins/Developer/Concert/ConcertUI/ConcertClientSharedSlate/Source/ConcertClientSharedSlate/Private/Replication/Editor/View/ObjectViewer/SSubobjectAndPropertySection.cpp b/Engine/Plugins/Developer/Concert/ConcertUI/ConcertClientSharedSlate/Source/ConcertClientSharedSlate/Private/Replication/Editor/View/ObjectViewer/SSubobjectAndPropertySection.cpp
deleted file mode 100644
diff --git a/Engine/Plugins/Developer/Concert/ConcertUI/ConcertClientSharedSlate/Source/ConcertClientSharedSlate/Private/Replication/Editor/View/ObjectViewer/SSubobjectAndPropertySection.h b/Engine/Plugins/Developer/Concert/ConcertUI/ConcertClientSharedSlate/Source/ConcertClientSharedSlate/Private/Replication/Editor/View/ObjectViewer/SSubobjectAndPropertySection.h
deleted file mode 100644
diff --git a/Engine/Plugins/Developer/Concert/ConcertUI/ConcertClientSharedSlate/Source/ConcertClientSharedSlate/Private/Replication/Editor/View/ObjectViewer/Tree/SReplicationColumnRow.h b/Engine/Plugins/Developer/Concert/ConcertUI/ConcertClientSharedSlate/Source/ConcertClientSharedSlate/Private/Replication/Editor/View/ObjectViewer/Tree/SReplicationColumnRow.h
deleted file mode 100644
diff --git a/Engine/Plugins/Developer/Concert/ConcertUI/ConcertClientSharedSlate/Source/ConcertClientSharedSlate/Private/Replication/Editor/View/ObjectViewer/Tree/SReplicationTreeView.h b/Engine/Plugins/Developer/Concert/ConcertUI/ConcertClientSharedSlate/Source/ConcertClientSharedSlate/Private/Replication/Editor/View/ObjectViewer/Tree/SReplicationTreeView.h
deleted file mode 100644
diff --git a/Engine/Plugins/Developer/Concert/ConcertUI/ConcertClientSharedSlate/Source/ConcertClientSharedSlate/Private/Replication/Editor/View/ObjectViewer/Tree/SelectionViewerColumns.cpp b/Engine/Plugins/Developer/Concert/ConcertUI/ConcertClientSharedSlate/Source/ConcertClientSharedSlate/Private/Replication/Editor/View/ObjectViewer/Tree/SelectionViewerColumns.cpp
deleted file mode 100644
diff --git a/Engine/Plugins/Developer/Concert/ConcertUI/ConcertClientSharedSlate/Source/ConcertClientSharedSlate/Private/Replication/Editor/View/ObjectViewer/Tree/SelectionViewerColumns.h b/Engine/Plugins/Developer/Concert/ConcertUI/ConcertClientSharedSlate/Source/ConcertClientSharedSlate/Private/Replication/Editor/View/ObjectViewer/Tree/SelectionViewerColumns.h
deleted file mode 100644
diff --git a/Engine/Plugins/Developer/Concert/ConcertUI/ConcertClientSharedSlate/Source/ConcertClientSharedSlate/Private/Replication/Editor/View/SubobjectView/ReplicatedSubobjectData.h b/Engine/Plugins/Developer/Concert/ConcertUI/ConcertClientSharedSlate/Source/ConcertClientSharedSlate/Private/Replication/Editor/View/SubobjectView/ReplicatedSubobjectData.h
deleted file mode 100644
diff --git a/Engine/Plugins/Developer/Concert/ConcertUI/ConcertClientSharedSlate/Source/ConcertClientSharedSlate/Private/Replication/Editor/View/SubobjectView/SSubobjectView.cpp b/Engine/Plugins/Developer/Concert/ConcertUI/ConcertClientSharedSlate/Source/ConcertClientSharedSlate/Private/Replication/Editor/View/SubobjectView/SSubobjectView.cpp
deleted file mode 100644
diff --git a/Engine/Plugins/Developer/Concert/ConcertUI/ConcertClientSharedSlate/Source/ConcertClientSharedSlate/Private/Replication/Editor/View/SubobjectView/SSubobjectView.h b/Engine/Plugins/Developer/Concert/ConcertUI/ConcertClientSharedSlate/Source/ConcertClientSharedSlate/Private/Replication/Editor/View/SubobjectView/SSubobjectView.h
deleted file mode 100644
diff --git a/Engine/Plugins/Developer/Concert/ConcertUI/ConcertClientSharedSlate/Source/ConcertClientSharedSlate/Private/Replication/Editor/View/Tree/SReplicationColumnRow.h b/Engine/Plugins/Developer/Concert/ConcertUI/ConcertClientSharedSlate/Source/ConcertClientSharedSlate/Private/Replication/Editor/View/Tree/SReplicationColumnRow.h
new file mode 100644
--- /dev/null
+++ b/Engine/Plugins/Developer/Concert/ConcertUI/ConcertClientSharedSlate/Source/ConcertClientSharedSlate/Private/Replication/Editor/View/Tree/SReplicationColumnRow.h
@@ -0,0 +1,129 @@
+// Copyright Epic Games, Inc. All Rights Reserved.
+
+#pragma once
+
+#include "Replication/Editor/View/ReplicationColumn.h"
+#include "Widgets/Layout/SBox.h"
+#include "Widgets/SNullWidget.h"
+#include "Widgets/Views/SExpanderArrow.h"
+#include "Widgets/Views/STableRow.h"
+#include "Widgets/Views/STableViewBase.h"
+
+namespace UE::ConcertClientSharedSlate
+{
+	/**
+	 * Displays the summary of an activity recorded and recoverable in the SConcertSessionRecovery list view.
+	 */
+	template<typename TListItemType>
+	class SReplicationColumnRow : public SMultiColumnTableRow<TSharedPtr<TListItemType>>
+	{
+	public:
+
+		using TColumType = TReplicationColumn<TListItemType>;
+		
+		DECLARE_DELEGATE_RetVal_OneParam(const TColumType*, FGetColumn,
+			const FName& ColumnId
+			);
+		DECLARE_DELEGATE_RetVal_TwoParams(TSharedPtr<SWidget>, FOverrideColumnWidget, const FName& ColumnName, const TListItemType& RowData);
+		
+		SLATE_BEGIN_ARGS(SReplicationColumnRow)
+			: _RowHeight(20.f)
+		{}
+			/** Used for highlighting the text being searched. */
+			SLATE_ARGUMENT(TSharedPtr<FText>, HighlightText)
+
+			/** Gets columns info about a certain column */
+			SLATE_EVENT(FGetColumn, ColumnGetter)
+
+			/**
+			 * Optional. If the delegate returns non-null, that widget will be used instead of the one the column would generate.
+			 * This is useful, e.g. if you want to generate a separator widget between items.
+			 */
+			SLATE_EVENT(FOverrideColumnWidget, OverrideColumnWidget)
+
+			/** The data to pass to TReplicationColumn::BuildColumnWidget. */
+			SLATE_ARGUMENT(TSharedPtr<TListItemType>, RowData)
+		
+			/** The name of the column that will have the SExpanderArrow for the tree view. */
+			SLATE_ARGUMENT(FName, ExpandableColumnLabel)
+
+			/** The height of the row */
+			SLATE_ARGUMENT(float, RowHeight)
+		SLATE_END_ARGS()
+
+		void Construct(
+			const FArguments& InArgs,
+			TSharedRef<STableViewBase> InOwner)
+		{
+			ColumnGetterDelegate = InArgs._ColumnGetter;
+			OverrideColumnWidgetDelegate = InArgs._OverrideColumnWidget;
+			HighlightText = InArgs._HighlightText;
+			RowData = InArgs._RowData;
+			ExpandableColumnLabel = InArgs._ExpandableColumnLabel;
+			RowHeight = InArgs._RowHeight;
+			
+			SMultiColumnTableRow<TSharedPtr<TListItemType>>::Construct({}, InOwner);
+		}
+
+		/** Generates the widget representing this row. */
+		virtual TSharedRef<SWidget> GenerateWidgetForColumn(const FName& ColumnName) override
+		{
+			const TSharedPtr<SWidget> ColumnOverride = OverrideColumnWidgetDelegate.IsBound()
+				? OverrideColumnWidgetDelegate.Execute(ColumnName, *RowData.Get())
+				: nullptr;
+			if (ColumnOverride)
+			{
+				return ColumnOverride.ToSharedRef();
+			}
+			
+			const TColumType* Column = ColumnGetterDelegate.Execute(ColumnName); ensure(Column);
+			if (!Column)
+			{
+				return SNullWidget::NullWidget;
+			}
+			
+			const TSharedRef<SWidget> ColumnWidget = Column->BuildColumnWidget({ HighlightText, *RowData.Get() });
+			const bool bNeedsExpanderArrow = ColumnName == ExpandableColumnLabel;
+			if (!bNeedsExpanderArrow)
+			{
+				return SNew(SBox)
+					.MinDesiredHeight(RowHeight)
+					.VAlign(VAlign_Center)
+					[
+						ColumnWidget
+					];
+			}
+			
+			return SNew(SBox)
+				.MinDesiredHeight(RowHeight)
+				[
+					SNew(SHorizontalBox)
+
+					+SHorizontalBox::Slot()
+					.AutoWidth()
+					.Padding(6, 0, 0, 0)
+					[
+						SNew(SExpanderArrow, SReplicationColumnRow::SharedThis(this))
+						.IndentAmount(12)
+					]
+
+					+SHorizontalBox::Slot()
+					.FillWidth(1.0f)
+					.VAlign(VAlign_Center)
+					[
+						ColumnWidget
+					]
+				];
+	
+		}
+
+	private:
+		
+		FGetColumn ColumnGetterDelegate;
+		FOverrideColumnWidget OverrideColumnWidgetDelegate;
+		TSharedPtr<FText> HighlightText;
+		TSharedPtr<TListItemType> RowData;
+		FName ExpandableColumnLabel;
+		float RowHeight;
+	};
+}
\ No newline at end of file
diff --git a/Engine/Plugins/Developer/Concert/ConcertUI/ConcertClientSharedSlate/Source/ConcertClientSharedSlate/Private/Replication/Editor/View/Tree/SReplicationTreeView.h b/Engine/Plugins/Developer/Concert/ConcertUI/ConcertClientSharedSlate/Source/ConcertClientSharedSlate/Private/Replication/Editor/View/Tree/SReplicationTreeView.h
new file mode 100644
--- /dev/null
+++ b/Engine/Plugins/Developer/Concert/ConcertUI/ConcertClientSharedSlate/Source/ConcertClientSharedSlate/Private/Replication/Editor/View/Tree/SReplicationTreeView.h
@@ -0,0 +1,419 @@
+// Copyright Epic Games, Inc. All Rights Reserved.
+
+#pragma once
+
+#include "Replication/Editor/View/ReplicationColumn.h"
+#include "SReplicationColumnRow.h"
+
+#include "Algo/RemoveIf.h"
+#include "Misc/TextFilter.h"
+#include "Widgets/DeclarativeSyntaxSupport.h"
+#include "Widgets/Input/SSearchBox.h"
+#include "Widgets/Layout/SScrollBox.h"
+#include "Widgets/Layout/SWidgetSwitcher.h"
+#include "Widgets/SBoxPanel.h"
+#include "Widgets/Views/STreeView.h"
+
+#define LOCTEXT_NAMESPACE "SReplicationListView"
+
+namespace UE::ConcertClientSharedSlate
+{
+	/**
+	 * Shared code for the list view for replicated actors and properties.
+	 * It is a table view that is searchable with a search box and exposes slots to add more filter widgets, such as SBasicFilterBar.
+	 */
+	template<typename TItemType>
+	class SReplicationTreeView : public SCompoundWidget
+	{
+	public:
+
+		using TOverrideColumnWidget = typename SReplicationColumnRow<TItemType>::FOverrideColumnWidget;
+
+		DECLARE_DELEGATE_OneParam(FDeleteItems, const TArray<TSharedPtr<TItemType>>& SelectedItems);
+		DECLARE_DELEGATE_TwoParams(FGetItemChildren, TSharedPtr<TItemType> Item, TFunctionRef<void(TSharedPtr<TItemType>)> ProcessChild);
+		DECLARE_DELEGATE(FOnSelectionChanged);
+		
+		DECLARE_DELEGATE_RetVal_OneParam(bool, FCustomFilter, const TSharedPtr<TItemType>& Item);
+		DECLARE_DELEGATE_RetVal_OneParam(bool, FIsSearchableItem, const TSharedPtr<TItemType>& Item);
+
+		enum class EContent
+		{
+			TreeView,
+			Custom
+		};
+
+		SLATE_BEGIN_ARGS(SReplicationTreeView<TItemType>)
+			: _HeaderRowVisibility(EVisibility::Visible)
+			, _SelectionMode(ESelectionMode::Single)
+		{}
+			/** The items to display */
+			SLATE_ARGUMENT(TArray<TSharedPtr<TItemType>>*, RootItemsSource)
+
+			/** Gets an items children for the tree view */
+			SLATE_EVENT(FGetItemChildren, OnGetChildren)
+
+			/** Optional. Called when the user presses the delete key */
+			SLATE_EVENT(FDeleteItems, OnDeleteItems)
+
+			/** Called to generate the context menu for an item */
+			SLATE_EVENT(FOnContextMenuOpening, OnContextMenuOpening)
+
+			/** Called when the selection changes. Call GetSelectedItems to get the selected items. */
+			SLATE_EVENT(FOnSelectionChanged, OnSelectionChanged)
+
+			/** Optional callback to do even more filtering of items. */
+			SLATE_EVENT(FCustomFilter, FilterItem)
+		
+			/**
+			 * Optional. If the delegate returns non-null, that widget will be used instead of the one the column would generate.
+			 * This is useful, e.g. if you want to generate a separator widget between items.
+			 */
+			SLATE_EVENT(TOverrideColumnWidget, OverrideColumnWidget)
+			/** Optional callback for determining whether this item can be searched. */
+			SLATE_EVENT(FIsSearchableItem, IsSearchableItem)
+			
+			/** The columns this list should have */
+			SLATE_ARGUMENT(TArray<TReplicationColumn<TItemType>>, Columns)
+			/** The name of the column that will have the SExpanderArrow for the tree view. */
+			SLATE_ARGUMENT(FName, ExpandableColumnLabel)
+			/** Visibility of the header row */
+			SLATE_ARGUMENT(EVisibility, HeaderRowVisibility)
+		
+			/** How many items are to allowed to be selected */
+			SLATE_ARGUMENT(ESelectionMode::Type, SelectionMode)
+
+			/** Optional widget to add to the left of the search bar. */
+			SLATE_NAMED_SLOT(FArguments, LeftOfSearchBar)
+			/** Optional widget to add between the search bar and the table view (e.g. a SBasicFilterBar). */
+			SLATE_NAMED_SLOT(FArguments, RowBelowSearchBar)
+
+			/** Optional, alternate content to show instead of the tree view when there are no rows. */
+			SLATE_NAMED_SLOT(FArguments, NoItemsContent)
+		SLATE_END_ARGS()
+
+		void Construct(const FArguments& InArgs)
+		{
+			AllRootItems = InArgs._RootItemsSource;
+			check(AllRootItems);
+
+			OnGetChildrenDelegate = InArgs._OnGetChildren;
+			OnDeleteItemsDelegate = InArgs._OnDeleteItems;
+			CustomFilterDelegate = InArgs._FilterItem;
+			OverrideColumnWidget = InArgs._OverrideColumnWidget;
+			IsSearchableItemDelegate = InArgs._IsSearchableItem;
+			ExpandableColumnId = InArgs._ExpandableColumnLabel;
+			
+			SearchText = MakeShared<FText>();
+			SearchTextFilter = MakeShared<TTextFilter<const TSharedPtr<TItemType>&>>(TTextFilter<const TSharedPtr<TItemType>&>::FItemToStringArray::CreateSP(this, &SReplicationTreeView::PopulateSearchStrings));
+			SearchTextFilter->OnChanged().AddSP(this, &SReplicationTreeView::OnFilterChanged);
+			
+			ChildSlot
+			[
+				SNew(SVerticalBox)
+
+				// Search
+				+SVerticalBox::Slot()
+				.Padding(1.f)
+				.AutoHeight()
+				[
+					SNew(SHorizontalBox)
+
+					+SHorizontalBox::Slot()
+					.Padding(1.f)
+					.AutoWidth()
+					[
+						InArgs._LeftOfSearchBar.Widget
+					]
+
+					+SHorizontalBox::Slot()
+					.Padding(1.f)
+					.FillWidth(1.f)
+					[
+						SAssignNew(SearchBox, SSearchBox)
+						.HintText(LOCTEXT("SearchHint", "Search..."))
+						.OnTextChanged(this, &SReplicationTreeView::OnSearchTextChanged)
+						.OnTextCommitted(this, &SReplicationTreeView::OnSearchTextCommitted)
+						.DelayChangeNotificationsWhileTyping(true)
+					]
+				]
+
+				// Optional slot between search bar and table, e.g. for an external SBasicFilterBar
+				+SVerticalBox::Slot()
+				.Padding(1.f)
+				.AutoHeight()
+				[
+					InArgs._RowBelowSearchBar.Widget
+				]
+
+				// Table row
+				+SVerticalBox::Slot()
+				.Padding(1.f)
+				.FillHeight(1.f)
+				[
+					SNew(SScrollBox)
+					+SScrollBox::Slot()
+					.FillSize(1.f)
+					[
+						CreateTreeView(InArgs)
+					]
+				]
+			];
+		}
+		
+		void OnItemsChanged()
+		{
+			// Re-filter everything. There should not be many items so filtering everything again should be fine. Calls RequestListRefresh as well.
+			OnFilterChanged();
+		}
+		
+		void SetSelectedItems(const TArray<TSharedPtr<TItemType>>& ObjectsToSelect, bool bIsSelected)
+		{
+			TreeView->ClearSelection();
+			TreeView->SetItemSelection(ObjectsToSelect, bIsSelected);
+		}
+		void SetExpandedItems(const TArray<TSharedPtr<TItemType>>& ObjectsToSelect, bool bIsExpanded)
+		{
+			for (TSharedPtr<TItemType> Item : ObjectsToSelect)
+			{
+				TreeView->SetItemExpansion(Item, bIsExpanded);
+			}
+		}
+		
+		TArray<TSharedPtr<TItemType>> GetSelectedItems() const { return TreeView->GetSelectedItems(); }
+		const TArray<TSharedPtr<TItemType>>& GetFilteredRootItems() const { return FilteredRootItems; }
+		
+		virtual FReply OnKeyDown(const FGeometry& MyGeometry, const FKeyEvent& InKeyEvent) override;
+
+	private:
+
+		/** The widget being used for search. */
+		TSharedPtr<SSearchBox> SearchBox;
+		/** Used to highlight text in text widgets */
+		TSharedPtr<FText> SearchText;
+		/** Performs text search */
+		TSharedPtr<TTextFilter<const TSharedPtr<TItemType>&>> SearchTextFilter;
+
+		/** ListView's header row */
+		TSharedPtr<SHeaderRow> HeaderRow;
+		/** Displays the contents */
+		TSharedPtr<STreeView<TSharedPtr<TItemType>>> TreeView;
+		/** The name of the column which will have the SExpandableArrow widget for the tree view. */
+		FName ExpandableColumnId;
+
+		TArray<TSharedPtr<TItemType>>* AllRootItems = nullptr;
+		TArray<TSharedPtr<TItemType>> FilteredRootItems;
+
+		/** Callback for getting an item's children. */
+		FGetItemChildren OnGetChildrenDelegate;
+		/** Optional delegate for responding to pressing the delete button */
+		FDeleteItems OnDeleteItemsDelegate;
+		/** Optional delegate for filtering the items even more. */
+		FCustomFilter CustomFilterDelegate;
+		/** Optional delegate for overriding the column widgets. */
+		TOverrideColumnWidget OverrideColumnWidget;
+		/** Optional callback for determining whether this item can be filtered. If false, it will not be shown when searched. */
+		FIsSearchableItem IsSearchableItemDelegate;
+		
+
+		TSharedRef<SWidget> CreateTreeView(const FArguments& InArgs);
+		TSharedRef<SHeaderRow> CreateHeaderRow(const FArguments& InArgs);
+		TSharedRef<ITableRow> OnGenerateRowWidget(TSharedPtr<TItemType> Item, const TSharedRef<STableViewBase>& OwnerTable);
+		void GetRowChildren(TSharedPtr<TItemType> Item, TArray<TSharedPtr<TItemType>>& OutChildren);
+		
+		void OnSearchTextCommitted(const FText& InFilterText, ETextCommit::Type CommitType);
+		void OnSearchTextChanged(const FText& InSearchText);
+
+		void PopulateSearchStrings(const TSharedPtr<TItemType>& Item, TArray<FString>& OutSearchStrings);
+		void OnFilterChanged();
+		bool PassesFilters(const TSharedPtr<TItemType>& Item);
+	};
+
+	template <typename TItemType>
+	FReply SReplicationTreeView<TItemType>::OnKeyDown(const FGeometry& MyGeometry, const FKeyEvent& InKeyEvent)
+	{
+		if (InKeyEvent.GetKey() == EKeys::Delete && OnDeleteItemsDelegate.IsBound())
+		{
+			OnDeleteItemsDelegate.Execute(TreeView->GetSelectedItems());
+			return FReply::Handled();
+		}
+	
+		return SCompoundWidget::OnKeyDown(MyGeometry, InKeyEvent);
+	}
+
+	template <typename TItemType>
+	TSharedRef<SWidget> SReplicationTreeView<TItemType>::CreateTreeView(const FArguments& InArgs)
+	{
+		TSharedPtr<SVerticalBox> VerticalBox;
+		
+		TSharedRef<SWidget> Result = SNew(SBorder)
+			.BorderImage(FAppStyle::Get().GetBrush("ToolPanel.GroupBorder"))
+			.BorderBackgroundColor(FSlateColor(FLinearColor(0.6, 0.6, 0.6)))
+			.Padding(0)
+			[
+				SAssignNew(VerticalBox, SVerticalBox)
+				
+				+SVerticalBox::Slot()
+				.FillHeight(1.f)
+				[
+					SAssignNew(TreeView, STreeView<TSharedPtr<TItemType>>)
+					.OnGetChildren(this, &SReplicationTreeView::GetRowChildren)
+					.TreeItemsSource(&FilteredRootItems)
+					.OnGenerateRow(this, &SReplicationTreeView::OnGenerateRowWidget)
+					.OnContextMenuOpening(InArgs._OnContextMenuOpening)
+					.OnSelectionChanged_Lambda([OnSelectionChanged = InArgs._OnSelectionChanged](auto, auto){ OnSelectionChanged.ExecuteIfBound(); })
+					.SelectionMode(InArgs._SelectionMode)
+					.AllowOverscroll(EAllowOverscroll::No)
+					.HeaderRow(CreateHeaderRow(InArgs))
+				]
+
+				+SVerticalBox::Slot()
+				.FillHeight(1.f)
+				.HAlign(HAlign_Center)
+				.VAlign(VAlign_Center)
+				.Padding(0.f, 0.f, 0.f, 20.f)
+				[
+					SNew(SWidgetSwitcher)
+					.WidgetIndex_Lambda([this](){ return AllRootItems->IsEmpty() ? 1 : 0; })
+					.Visibility_Lambda([this](){ return AllRootItems->IsEmpty() || FilteredRootItems.IsEmpty() ? EVisibility::Visible : EVisibility::Collapsed; })
+					+SWidgetSwitcher::Slot() [ SNew(STextBlock).Text(LOCTEXT("AllFiltered", "All items are filtered.")) ]
+					+SWidgetSwitcher::Slot()
+					[
+						InArgs._NoItemsContent.Widget
+					]
+				]
+			];
+		
+		return Result;
+	}
+
+	template <typename TItemType>
+	TSharedRef<SHeaderRow> SReplicationTreeView<TItemType>::CreateHeaderRow(const FArguments& InArgs)
+	{
+		TArray<TReplicationColumn<TItemType>> Columns = InArgs._Columns;
+		Columns.Sort([](const TReplicationColumn<TItemType>& Left, const TReplicationColumn<TItemType>& Right) { return Left.GetColumnSortOrderValue() < Right.GetColumnSortOrderValue(); });
+		
+		HeaderRow = SNew(SHeaderRow).Visibility(InArgs._HeaderRowVisibility);
+		TSet<FName> DuplicateColumnDetection;
+		for (TReplicationColumn<TItemType>& Column : Columns)
+		{
+			check(!DuplicateColumnDetection.Contains(Column.ColumnId));
+			DuplicateColumnDetection.Add(Column.ColumnId);
+			
+			// SHeaderRow owns the columns and deletes them when destroyed
+			TReplicationColumn<TItemType>* ManagedByHeaderRow = new TReplicationColumn<TItemType>(MoveTemp(Column));
+			HeaderRow->AddColumn(*ManagedByHeaderRow);
+		}
+
+		return HeaderRow.ToSharedRef();
+	}
+
+	template <typename TItemType>
+	TSharedRef<ITableRow> SReplicationTreeView<TItemType>::OnGenerateRowWidget(TSharedPtr<TItemType> Item, const TSharedRef<STableViewBase>& OwnerTable)
+	{
+		const typename SReplicationColumnRow<TItemType>::FGetColumn ColumnGetter = SReplicationColumnRow<TItemType>::FGetColumn::CreateLambda([this](const FName& ColumnId) -> const TReplicationColumn<TItemType>*
+		{
+			for (const SHeaderRow::FColumn& Column : HeaderRow->GetColumns())
+			{
+				if (Column.ColumnId == ColumnId)
+				{
+					return static_cast<const TReplicationColumn<TItemType>*>(&Column);
+				}
+			}
+
+			checkNoEntry();
+			return nullptr;
+		});
+	
+		return SNew(SReplicationColumnRow<TItemType>, OwnerTable)
+			.HighlightText(SearchText)
+			.ColumnGetter(ColumnGetter)
+			.OverrideColumnWidget(OverrideColumnWidget)
+			.RowData(Item)
+			.ExpandableColumnLabel(ExpandableColumnId);
+	}
+
+	template <typename TItemType>
+	void SReplicationTreeView<TItemType>::GetRowChildren(TSharedPtr<TItemType> Item, TArray<TSharedPtr<TItemType>>& OutChildren)
+	{
+		if (OnGetChildrenDelegate.IsBound())
+		{
+			OnGetChildrenDelegate.Execute(Item, [this, &OutChildren](TSharedPtr<TItemType> ItemToAdd)
+			{
+				if (PassesFilters(ItemToAdd))
+				{
+					OutChildren.Add(ItemToAdd);
+				}
+			});
+		}
+	}
+
+	template <typename TItemType>
+	void SReplicationTreeView<TItemType>::OnSearchTextCommitted(const FText& InFilterText, ETextCommit::Type CommitType)
+	{
+		if (!InFilterText.EqualTo(*SearchText))
+		{
+			OnSearchTextChanged(InFilterText);
+		}
+	}
+
+	template <typename TItemType>
+	void SReplicationTreeView<TItemType>::OnSearchTextChanged(const FText& InSearchText)
+	{
+		*SearchText = InSearchText;
+		SearchTextFilter->SetRawFilterText(InSearchText);
+		SearchBox->SetError(SearchTextFilter->GetFilterErrorText());
+	}
+
+	template <typename TItemType>
+	void SReplicationTreeView<TItemType>::PopulateSearchStrings(const TSharedPtr<TItemType>& Item, TArray<FString>& OutSearchStrings)
+	{
+		if (IsSearchableItemDelegate.IsBound() && !IsSearchableItemDelegate.Execute(Item))
+		{
+			return;
+		}
+	
+		for (const SHeaderRow::FColumn& Column : HeaderRow->GetColumns())
+		{
+			const TReplicationColumn<TItemType>& CastColumn = static_cast<const TReplicationColumn<TItemType>&>(Column);
+			CastColumn.ExecutePopulateSearchString(*Item.Get(), OutSearchStrings);
+		}
+	}
+
+	template <typename TItemType>
+	void SReplicationTreeView<TItemType>::OnFilterChanged()
+	{
+		// Try preserving the selected activity.
+		TArray<TSharedPtr<TItemType>> SelectedItems = TreeView->GetSelectedItems();
+
+		// Reset the list of displayed activities.
+		FilteredRootItems.Reset(AllRootItems->Num());
+
+		// Apply the filter.
+		for (const TSharedPtr<TItemType>& Activity : *AllRootItems)
+		{
+			if (PassesFilters(Activity))
+			{
+				FilteredRootItems.Add(Activity);
+			}
+		}
+
+		// Restore/reset the selected activity.
+		SelectedItems.SetNum(Algo::RemoveIf(SelectedItems, [this](const TSharedPtr<TItemType>& Item){ return !FilteredRootItems.Contains(Item); }));
+		if (!SelectedItems.IsEmpty())
+		{
+			TreeView->SetItemSelection(SelectedItems, true); // Restore previous selection.
+			TreeView->RequestScrollIntoView(SelectedItems[0]);
+		}
+
+		TreeView->RequestListRefresh();
+	}
+
+	template <typename TItemType>
+	bool SReplicationTreeView<TItemType>::PassesFilters(const TSharedPtr<TItemType>& Item)
+	{
+		return SearchTextFilter->PassesFilter(Item)
+			&& (!CustomFilterDelegate.IsBound() || CustomFilterDelegate.Execute(Item));
+	}
+}
+
+#undef LOCTEXT_NAMESPACE
\ No newline at end of file
diff --git a/Engine/Plugins/Developer/Concert/ConcertUI/ConcertClientSharedSlate/Source/ConcertClientSharedSlate/Private/Replication/Editor/View/Tree/SelectionViewerColumns.cpp b/Engine/Plugins/Developer/Concert/ConcertUI/ConcertClientSharedSlate/Source/ConcertClientSharedSlate/Private/Replication/Editor/View/Tree/SelectionViewerColumns.cpp
new file mode 100644
--- /dev/null
+++ b/Engine/Plugins/Developer/Concert/ConcertUI/ConcertClientSharedSlate/Source/ConcertClientSharedSlate/Private/Replication/Editor/View/Tree/SelectionViewerColumns.cpp
@@ -0,0 +1,315 @@
+// Copyright Epic Games, Inc. All Rights Reserved.
+
+#include "SelectionViewerColumns.h"
+
+#include "ClassIconFinder.h"
+#include "Replication/Editor/Model/IEditableReplicationStreamModel.h"
+#include "Replication/Editor/Model/ISubobjectModel.h"
+#include "Replication/Editor/Model/ReplicatedPropertyData.h"
+#include "Replication/Editor/Model/ReplicatedObjectData.h"
+#include "Replication/Editor/View/DisplayUtils.h"
+#include "Replication/Editor/View/ObjectEditor/SBaseReplicationStreamEditor.h"
+#include "Replication/Editor/View/ReplicationColumnsUtils.h"
+#include "Replication/PropertyChainUtils.h"
+
+#include "Internationalization/Internationalization.h"
+#include "GameFramework/Actor.h"
+#include "Replication/Editor/View/ObjectUtils.h"
+#include "Textures/SlateIcon.h"
+#include "Widgets/Images/SImage.h"
+#include "Widgets/SBoxPanel.h"
+#include "Widgets/Text/STextBlock.h"
+#include "Widgets/Input/SCheckBox.h"
+
+#define LOCTEXT_NAMESPACE "ReplicationObjectColumns"
+
+namespace UE::ConcertClientSharedSlate::ReplicationColumns::TopLevel
+{
+	const FName IconColumnId = TEXT("IconColumn");
+	const FName LabelColumnId = TEXT("LabelColumn");
+	const FName TypeColumnId = TEXT("TypeColumn");
+	
+	FReplicationTopLevelObjectColumn LabelColumn(TSharedRef<IReplicationStreamModel> Model, ISubobjectModel* SubobjectModel)
+	{
+		return FReplicationTopLevelObjectColumn(
+			FReplicationTopLevelObjectColumn::FArguments()
+				.GenerateWidgetColumn_Lambda([Model = MoveTemp(Model), SubobjectModel](const FReplicationTopLevelObjectColumn::FBuildArgs& Args)
+				{
+					const FSoftObjectPath& ObjectPath = Args.RowData.GetObjectPath();
+					FText Text = DisplayUtils::GetObjectDisplayText(ObjectPath);
+					if (SubobjectModel && !SubobjectModel->IsTopLevelObject(ObjectPath))
+					{
+						if (const TOptional<FSoftObjectPath> OwningActorPath = ObjectUtils::GetActorOf(ObjectPath))
+						{
+							SubobjectModel->SetTopLevelObject(*OwningActorPath);
+							Text = SubobjectModel->GetSubobjectDisplayName(ObjectPath);
+						}
+					}
+					
+					return SNew(SHorizontalBox)
+						+SHorizontalBox::Slot()
+						.AutoWidth()
+						.HAlign(HAlign_Center)
+						.VAlign(VAlign_Center)
+						[
+							SNew(SImage)
+							.Image(DisplayUtils::GetObjectIcon(*Model, Args.RowData.GetObjectPath()).GetOptionalIcon())
+						]
+					
+						+SHorizontalBox::Slot()
+						.HAlign(HAlign_Left)
+						.VAlign(VAlign_Center)
+						.Padding(6.f, 0.f, 0.f, 0.f)
+						[
+							SNew(STextBlock)
+							.HighlightText(TAttribute<FText>::CreateLambda([HighlightText = Args.HighlightText](){ return *HighlightText; }))
+							.Text(Text)
+						];
+				})
+				.PopulateSearchItems_Lambda([](const FReplicatedObjectData& ObjectData, TArray<FString>& InOutSearchStrings)
+				{
+					InOutSearchStrings.Add(DisplayUtils::GetObjectDisplayText(ObjectData.GetObjectPath()).ToString());
+				})
+				.ColumnSortOrder(static_cast<int32>(ETopLevelColumnOrder::Label)),
+			SHeaderRow::Column(LabelColumnId)
+				.DefaultLabel(LOCTEXT("LabelColumnLabel", "Label"))
+				.FillWidth(1.f)
+			);
+	}
+	
+	FReplicationTopLevelObjectColumn TypeColumn(TSharedRef<IReplicationStreamModel> Model)
+	{
+		return FReplicationTopLevelObjectColumn(
+			FReplicationTopLevelObjectColumn::FArguments()
+				.GenerateWidgetColumn_Lambda([Model](const FReplicationTopLevelObjectColumn::FBuildArgs& Args)
+				{
+					return SNew(STextBlock)
+						.HighlightText(TAttribute<FText>::CreateLambda([HighlightText = Args.HighlightText](){ return *HighlightText; }))
+						.Text(DisplayUtils::GetObjectTypeText(*Model, Args.RowData.GetObjectPath()));
+				})
+				.PopulateSearchItems_Lambda([Model](const FReplicatedObjectData& ObjectData, TArray<FString>& InOutSearchStrings)
+				{
+					InOutSearchStrings.Add(DisplayUtils::GetObjectTypeText(Model.Get(), ObjectData.GetObjectPath()).ToString());
+				})
+				.ColumnSortOrder(static_cast<int32>(ETopLevelColumnOrder::Type)),
+			SHeaderRow::Column(TypeColumnId)
+				.DefaultLabel(LOCTEXT("TypeColumnLabel", "Type"))
+				.FillWidth(1.f)
+			);
+	}
+}
+
+#undef LOCTEXT_NAMESPACE
+#define LOCTEXT_NAMESPACE "ReplicationPropertyColumns"
+
+namespace UE::ConcertClientSharedSlate::ReplicationColumns::Property
+{
+	const FName ReplicatesColumnId = TEXT("ReplicatedColumn");
+	const FName LabelColumnId = TEXT("LabelColumn");
+	const FName TypeColumnId = TEXT("TypeColumn");
+	
+	FReplicationPropertyColumn LabelColumn()
+	{
+		return FReplicationPropertyColumn(
+			FReplicationPropertyColumn::FArguments()
+				.GenerateWidgetColumn_Lambda([](const FReplicationPropertyColumn::FBuildArgs& Args)
+				{
+					return SNew(STextBlock)
+						.HighlightText(TAttribute<FText>::CreateLambda([HighlightText = Args.HighlightText](){ return *HighlightText; }))
+						.Text(DisplayUtils::GetPropertyDisplayText(Args.RowData.GetProperty()));
+				})
+				.PopulateSearchItems_Lambda([](const FReplicatedPropertyData& ObjectData, TArray<FString>& InOutSearchStrings)
+				{
+					InOutSearchStrings.Add(DisplayUtils::GetPropertyDisplayText(ObjectData.GetProperty()).ToString());
+				})
+				.ColumnSortOrder(static_cast<int32>(EReplicationPropertyColumnOrder::Label)),
+			SHeaderRow::Column(LabelColumnId)
+				.DefaultLabel(LOCTEXT("LabelColumnLabel", "Label"))
+				.FillWidth(1.f)
+			);
+	}
+	
+	FReplicationPropertyColumn TypeColumn()
+	{
+		static auto GetDisplayText = [](const FReplicatedPropertyData& Args)
+		{
+			UClass* Class = Args.GetOwningClass().TryLoadClass<UObject>();
+			const FProperty* Property = Class ? ConcertSyncCore::PropertyChain::ResolveProperty(*Class, Args.GetProperty()) : nullptr;
+			return Property ? FText::FromString(Property->GetCPPType()) : LOCTEXT("Unknown", "Unknown");	
+		};
+		
+		return FReplicationPropertyColumn(
+			FReplicationPropertyColumn::FArguments()
+				.GenerateWidgetColumn_Lambda([](const FReplicationPropertyColumn::FBuildArgs& Args)
+				{
+					return SNew(STextBlock)
+						.HighlightText(TAttribute<FText>::CreateLambda([HighlightText = Args.HighlightText](){ return *HighlightText; }))
+						.Text(GetDisplayText(Args.RowData));
+				})
+				.PopulateSearchItems_Lambda([](const FReplicatedPropertyData& ObjectData, TArray<FString>& InOutSearchStrings)
+				{
+					InOutSearchStrings.Add(GetDisplayText(ObjectData).ToString());
+				})
+				.ColumnSortOrder(static_cast<int32>(EReplicationPropertyColumnOrder::Type)),
+			SHeaderRow::Column(TypeColumnId)
+				.DefaultLabel(LOCTEXT("TypeColumnLabel", "Type"))
+				.FillWidth(0.5f)
+			);
+	}
+
+	namespace Private
+	{
+		static ECheckBoxState OnGetPropertyCheckboxState(
+			const FConcertPropertyChain& PropertyChain,
+			const IReplicationStreamViewer& Viewer,
+			const IReplicationStreamModel& Model
+			)
+		{
+			const TArray<FSoftObjectPath> SelectedObjectPaths = Viewer.GetObjectsBeingPropertyEdited();
+			return GetPropertyCheckboxStateBasedOnSelection(PropertyChain, SelectedObjectPaths, Model);
+		}
+
+		static void OnPropertyCheckboxChanged(
+			bool bIsChecked,
+			const FConcertPropertyChain& PropertyChain,
+			const IReplicationStreamViewer& Viewer,
+			IEditableReplicationStreamModel& Model
+			)
+		{
+			const TArray Properties{ PropertyChain };
+			const TArray<FSoftObjectPath> SelectedObjects = Viewer.GetObjectsBeingPropertyEdited();
+
+			if (bIsChecked)
+			{
+				// We cannot proceed if any of the selected objects cannot be loaded because we cannot obtain its class
+				for (const FSoftObjectPath& Path : SelectedObjects)
+				{
+					UObject* Object = Path.ResolveObject();
+					if (!Object && !Model.ContainsObjects({ Path }))
+					{
+						return;
+					}
+				}
+				
+				for (const FSoftObjectPath& Path : Viewer.GetObjectsBeingPropertyEdited())
+				{
+					UObject* Object = Path.ResolveObject();
+					if (!Model.ContainsObjects({ Path }))
+					{
+						Model.AddObjects({ Object });
+					}
+					Model.AddProperties(Path, Properties);
+				}
+			}
+			else
+			{
+				for (const FSoftObjectPath& Path : Viewer.GetObjectsBeingPropertyEdited())
+				{
+					Model.RemoveProperties(Path, Properties);
+					// If this does not resolve it is not too bad if the subobject is not removed...
+					UObject* Object = Path.ResolveObject();
+					const bool bNeedsToRemoveNonRoot = Object && !Object->IsA<AActor>() && !Model.ContainsObjects({Path});
+					if (bNeedsToRemoveNonRoot)
+					{
+						Model.RemoveObjects({ Path });
+					}
+				}
+			}
+		}
+	}
+
+#define MOVE_CAPTURE(a)
+	
+	FReplicationPropertyColumn ReplicatesColumns(
+		TWeakPtr<IReplicationStreamViewer> Viewer,
+		TWeakPtr<IEditableReplicationStreamModel> Model,
+		TReplicationColumnDelegates<FReplicatedPropertyData>::FIsEnabled IsEnabledDelegate,
+		TAttribute<FText> DisabledToolTipText,
+		const float ColumnWidth,
+		const int32 Priority
+		)
+	{
+		using FPropertyColumnDelegates = TReplicationColumnDelegates<FReplicatedPropertyData>;
+		return MakeCheckboxColumn<FReplicatedPropertyData>(
+			ReplicatesColumnId,
+			FPropertyColumnDelegates(
+				FPropertyColumnDelegates::FGetColumnCheckboxState::CreateLambda(
+				[Viewer, Model](const FReplicatedPropertyData& Data)
+				{
+					const TSharedPtr<IReplicationStreamViewer> ViewerPin = Viewer.Pin();
+					const TSharedPtr<IEditableReplicationStreamModel> ModelPin = Model.Pin();
+					return ensure(ViewerPin && ModelPin) ? Private::OnGetPropertyCheckboxState(Data.GetProperty(), *ViewerPin, *ModelPin) : ECheckBoxState::Undetermined;
+				}),
+				FPropertyColumnDelegates::FOnColumnCheckboxChanged::CreateLambda(
+				[Viewer, Model](bool bIsChecked, const FReplicatedPropertyData& Data)
+				{
+					const TSharedPtr<IReplicationStreamViewer> ViewerPin = Viewer.Pin();
+					const TSharedPtr<IEditableReplicationStreamModel> ModelPin = Model.Pin();
+					if (ensure(ViewerPin && ModelPin))
+					{
+						Private::OnPropertyCheckboxChanged(bIsChecked, Data.GetProperty(), *ViewerPin, *ModelPin);
+					}
+				}),
+				FPropertyColumnDelegates::FGetToolTipText::CreateLambda(
+				[IsEnabledDelegate, DisabledToolTipText = MoveTemp(DisabledToolTipText)](const FReplicatedPropertyData& Data)
+				{
+					const bool bIsDisabled = IsEnabledDelegate.IsBound() && !IsEnabledDelegate.Execute(Data);
+					const bool bCanCall = DisabledToolTipText.IsBound() || DisabledToolTipText.IsSet();
+					return bIsDisabled
+						? bCanCall ? DisabledToolTipText.Get() : FText::GetEmpty()
+						: LOCTEXT("Replicates.ToolTip", "Select whether this property should be replicated");
+				}),
+				IsEnabledDelegate
+				),
+			FText::GetEmpty(),
+			Priority,
+			ColumnWidth
+		);
+	}
+
+	ECheckBoxState GetPropertyCheckboxStateBasedOnSelection(const FConcertPropertyChain& Property, TConstArrayView<FSoftObjectPath> Selection, const IReplicationStreamModel& Model)
+	{
+		ECheckBoxState CheckBoxState = ECheckBoxState::Undetermined;
+		for (const FSoftObjectPath& SelectedObject : Selection)
+		{
+			const bool bContainsProperty = Model.ContainsProperties(SelectedObject, { Property });
+			const ECheckBoxState StateForThisObject = bContainsProperty
+				? ECheckBoxState::Checked
+				: ECheckBoxState::Unchecked;
+			// The first object?
+			if (CheckBoxState == ECheckBoxState::Undetermined)
+			{
+				CheckBoxState = StateForThisObject;
+			}
+			else if (CheckBoxState != StateForThisObject)
+			{
+				// Two boxes do not have the same value
+				return ECheckBoxState::Undetermined;
+			}
+		}
+		return CheckBoxState;
+	}
+	
+	bool SortBySelectionThenByName_PropertyPredicate(
+		const TArray<FSoftObjectPath>& SelectedObjects,
+		const IReplicationStreamModel& Model,
+		const FReplicatedPropertyData& Left,
+		const FReplicatedPropertyData& Right
+		)
+	{
+		const ECheckBoxState LeftCheckboxState = GetPropertyCheckboxStateBasedOnSelection(Left.GetProperty(), SelectedObjects, Model);
+		const ECheckBoxState RightCheckboxState = GetPropertyCheckboxStateBasedOnSelection(Right.GetProperty(), SelectedObjects, Model);
+		
+		// Secondary sort by name
+		if (LeftCheckboxState == RightCheckboxState)
+		{
+			return DisplayUtils::GetPropertyDisplayString(Left.GetProperty()) < DisplayUtils::GetPropertyDisplayString(Right.GetProperty());
+		}
+
+		// Selected properties should appear first
+		return LeftCheckboxState == ECheckBoxState::Checked
+			&& (RightCheckboxState == ECheckBoxState::Unchecked || RightCheckboxState == ECheckBoxState::Undetermined);
+	}
+}
+
+#undef LOCTEXT_NAMESPACE
\ No newline at end of file
diff --git a/Engine/Plugins/Developer/Concert/ConcertUI/ConcertClientSharedSlate/Source/ConcertClientSharedSlate/Private/Replication/Editor/View/Tree/SelectionViewerColumns.h b/Engine/Plugins/Developer/Concert/ConcertUI/ConcertClientSharedSlate/Source/ConcertClientSharedSlate/Private/Replication/Editor/View/Tree/SelectionViewerColumns.h
new file mode 100644
--- /dev/null
+++ b/Engine/Plugins/Developer/Concert/ConcertUI/ConcertClientSharedSlate/Source/ConcertClientSharedSlate/Private/Replication/Editor/View/Tree/SelectionViewerColumns.h
@@ -0,0 +1,91 @@
+// Copyright Epic Games, Inc. All Rights Reserved.
+
+#pragma once
+
+#include "Replication/Editor/View/PredefinedReplicationColumns.h"
+#include "Replication/Editor/View/ReplicationColumnsUtils.h"
+#include "Templates/SharedPointer.h"
+
+namespace UE::ConcertClientSharedSlate
+{
+	class IEditableReplicationStreamModel;
+	class IReplicationStreamViewer;
+	class IReplicationStreamModel;
+	class ISubobjectModel;
+}
+
+namespace UE::ConcertClientSharedSlate::ReplicationColumns::TopLevel
+{
+	extern const FName IconColumnId;
+	extern const FName LabelColumnId;
+	extern const FName TypeColumnId;
+
+	enum class ETopLevelColumnOrder : int32
+	{
+		/** Label of the object */
+		Label = 20,
+		/** Class of the object */
+		Type = 30,
+	};
+
+	FReplicationTopLevelObjectColumn LabelColumn(TSharedRef<IReplicationStreamModel> Model, ISubobjectModel* SubobjectModel = nullptr);
+	FReplicationTopLevelObjectColumn TypeColumn(TSharedRef<IReplicationStreamModel> Model);
+}
+
+namespace UE::ConcertClientSharedSlate::ReplicationColumns::Property
+{
+	extern const FName ReplicatesColumnId;
+	extern const FName LabelColumnId;
+	extern const FName TypeColumnId;
+	
+	enum class EReplicationPropertyColumnOrder : int32
+	{
+		/** The checkbox in SDefaultReplicationStreamEditor determining whether the property is in the selection*/
+		ReplicatesCheckbox = 0,
+		/** Label of the property */
+		Label = 10,
+		/** Type of the property */
+		Type = 20
+	};
+	
+	FReplicationPropertyColumn LabelColumn();
+	FReplicationPropertyColumn TypeColumn();
+	/**
+	 * A checkbox that is placed at the beginning of every property.
+	 * Checking & unchecking adds & removes the property to the selected objects' property mapping, respectively.
+	 *
+	 * @param Viewer Used to determine which objects are currently selected (and being edited by the checkbox)
+	 * @param Model Used to actually edit the object to properties mapping
+	 * @param IsEnabledDelegate Determines whether the checkbox is enabled
+	 * @param DisabledToolTipText Tooltip to display when IsEnabledDelegate returns false
+	 * @param ColumnWidth Width to use for the column
+	 * @param Priority Determines position of this columns relative to the others
+	 */
+	CONCERTCLIENTSHAREDSLATE_API FReplicationPropertyColumn ReplicatesColumns(
+		TWeakPtr<IReplicationStreamViewer> Viewer,
+		TWeakPtr<IEditableReplicationStreamModel> Model,
+		TReplicationColumnDelegates<FReplicatedPropertyData>::FIsEnabled IsEnabledDelegate = {},
+		TAttribute<FText> DisabledToolTipText = {},
+		const float ColumnWidth = 20.f,
+		const int32 Priority = static_cast<int32>(EReplicationPropertyColumnOrder::ReplicatesCheckbox)
+		);
+
+	/**
+	 * Goes through all selected objects, checks whether the property is checked on it or not, and returns a checkbox state.
+	 * The state is checked or unchecked if all selected properties have the property enabled or disabled, respectively, and undetermined otherwise, i.e. if it is mixed.
+	 *
+	 * This is used by ReplicatesColumns.
+	 */
+	CONCERTCLIENTSHAREDSLATE_API ECheckBoxState GetPropertyCheckboxStateBasedOnSelection(
+		const FConcertPropertyChain& Property,
+		TConstArrayView<FSoftObjectPath> Selection,
+		const IReplicationStreamModel& Model
+		);
+	/** Util for sorting property data by whether its ReplicatesColumns() column is checked. */
+	CONCERTCLIENTSHAREDSLATE_API bool SortBySelectionThenByName_PropertyPredicate(
+		const TArray<FSoftObjectPath>& SelectedObjects,
+		const IReplicationStreamModel& Model,
+		const FReplicatedPropertyData& Left,
+		const FReplicatedPropertyData& Right
+		);
+}
diff --git a/Engine/Plugins/Developer/Concert/ConcertUI/ConcertClientSharedSlate/Source/ConcertClientSharedSlate/Private/Replication/ReplicationWidgetFactories.cpp b/Engine/Plugins/Developer/Concert/ConcertUI/ConcertClientSharedSlate/Source/ConcertClientSharedSlate/Private/Replication/ReplicationWidgetFactories.cpp
--- a/Engine/Plugins/Developer/Concert/ConcertUI/ConcertClientSharedSlate/Source/ConcertClientSharedSlate/Private/Replication/ReplicationWidgetFactories.cpp
+++ b/Engine/Plugins/Developer/Concert/ConcertUI/ConcertClientSharedSlate/Source/ConcertClientSharedSlate/Private/Replication/ReplicationWidgetFactories.cpp
@@ -7,16 +7,9 @@
 #include "Editor/View/MultiEditor/SMultiReplicationStreamEditor.h"
 #include "Editor/View/ObjectEditor/SDefaultReplicationStreamEditor.h"
 #include "Editor/View/ObjectEditor/SBaseReplicationStreamEditor.h"
-#include "Editor/View/SubobjectView/SSubobjectView.h"
 
 namespace UE::ConcertClientSharedSlate
 {
-	TSharedRef<IReplicationSubobjectView> CreateUnrealEditorSubobjectView(FCreateSubobjectViewParams Params)
-	{
-		return SNew(SSubobjectView, MoveTemp(Params.SubobjectModel))
-			.AdditionalColumns(Params.AdditionalColumns);
-	}
-	
 	TSharedRef<ISubobjectModel> CreateDefaultComponentHierarchySubobjectModel()
 	{
 		return MakeShared<FComponentHierarchySubobjectModel>();
@@ -27,7 +20,6 @@ namespace UE::ConcertClientSharedSlate
 		return SNew(SBaseReplicationStreamEditor, Params.DataModel, Params.ObjectSource, Params.PropertySource)
 			.AdditionalObjectColumns(Params.ViewerParams.AdditionalObjectColumns)
 			.AdditionalPropertyColumns(Params.ViewerParams.AdditionalPropertyColumns)
-			.SubobjectView(Params.ViewerParams.SubobjectView)
 			.SubobjectModel(Params.ViewerParams.SubobjectModel)
 			.OnExtendObjectsContextMenu(Params.ViewerParams.OnExtendObjectsContextMenu)
 			.SortPropertyRowPredicate(Params.ViewerParams.SortPropertyRowPredicate)
@@ -43,7 +35,6 @@ namespace UE::ConcertClientSharedSlate
 		return SNew(SDefaultReplicationStreamEditor, Params.DataModel, Params.ObjectSource, Params.PropertySource)
 			.AdditionalObjectColumns(Params.ViewerParams.AdditionalObjectColumns)
 			.AdditionalPropertyColumns(Params.ViewerParams.AdditionalPropertyColumns)
-			.SubobjectView(Params.ViewerParams.SubobjectView)
 			.SubobjectModel(Params.ViewerParams.SubobjectModel)
 			.OnExtendObjectsContextMenu(Params.ViewerParams.OnExtendObjectsContextMenu)
 			.SortPropertyRowPredicate(Params.ViewerParams.SortPropertyRowPredicate)
diff --git a/Engine/Plugins/Developer/Concert/ConcertUI/ConcertClientSharedSlate/Source/ConcertClientSharedSlate/Public/Replication/Editor/View/IReplicationSubobjectView.h b/Engine/Plugins/Developer/Concert/ConcertUI/ConcertClientSharedSlate/Source/ConcertClientSharedSlate/Public/Replication/Editor/View/IReplicationSubobjectView.h
deleted file mode 100644
diff --git a/Engine/Plugins/Developer/Concert/ConcertUI/ConcertClientSharedSlate/Source/ConcertClientSharedSlate/Public/Replication/Editor/View/PredefinedReplicationColumns.h b/Engine/Plugins/Developer/Concert/ConcertUI/ConcertClientSharedSlate/Source/ConcertClientSharedSlate/Public/Replication/Editor/View/PredefinedReplicationColumns.h
--- a/Engine/Plugins/Developer/Concert/ConcertUI/ConcertClientSharedSlate/Source/ConcertClientSharedSlate/Public/Replication/Editor/View/PredefinedReplicationColumns.h
+++ b/Engine/Plugins/Developer/Concert/ConcertUI/ConcertClientSharedSlate/Source/ConcertClientSharedSlate/Public/Replication/Editor/View/PredefinedReplicationColumns.h
@@ -10,7 +10,5 @@ namespace UE::ConcertClientSharedSlate::ReplicationColumns
 {
 	/** Rows displayed in the outliner */
 	using FReplicationTopLevelObjectColumn = TReplicationColumn<FReplicatedObjectData>;
-	/** The type of rows exposed by the subobject view as API. */
-	using FReplicationSubobjectObjectColumn = TReplicationColumn<FReplicatedObjectData>;
 	using FReplicationPropertyColumn = TReplicationColumn<FReplicatedPropertyData>;
 }
diff --git a/Engine/Plugins/Developer/Concert/ConcertUI/ConcertClientSharedSlate/Source/ConcertClientSharedSlate/Public/Replication/ReplicationWidgetFactories.h b/Engine/Plugins/Developer/Concert/ConcertUI/ConcertClientSharedSlate/Source/ConcertClientSharedSlate/Public/Replication/ReplicationWidgetFactories.h
--- a/Engine/Plugins/Developer/Concert/ConcertUI/ConcertClientSharedSlate/Source/ConcertClientSharedSlate/Public/Replication/ReplicationWidgetFactories.h
+++ b/Engine/Plugins/Developer/Concert/ConcertUI/ConcertClientSharedSlate/Source/ConcertClientSharedSlate/Public/Replication/ReplicationWidgetFactories.h
@@ -48,31 +48,9 @@ namespace UE::ConcertClientSharedSlate
 	/** Builds a similar tree hierarchy as SSubobjectEditor. Reports only components as subobjects. */
 	CONCERTCLIENTSHAREDSLATE_API TSharedRef<ISubobjectModel> CreateDefaultComponentHierarchySubobjectModel();
 	
-	struct FCreateSubobjectViewParams
-	{
-		/** Additional columns that should be displayed for the subobjects. */
-		TArray<ReplicationColumns::FReplicationSubobjectObjectColumn> AdditionalColumns;
-
-		/** Determines the subobjects displayed in the view. */
-		TSharedRef<ISubobjectModel> SubobjectModel = CreateDefaultComponentHierarchySubobjectModel();
-	};
-	/** Creates a stream editor with a subobject view that looks like the SSubobjectEditor. */
-	CONCERTCLIENTSHAREDSLATE_API TSharedRef<IReplicationSubobjectView> CreateUnrealEditorSubobjectView(
-		FCreateSubobjectViewParams Params
-		);
-
 	/** Params for creating a IReplicationStreamViewer. */
 	struct FCreateViewerParams
 	{
-		/**
-		 * Optional. This is inserted between the root object outliner and property view.
-		 * It e.g. displays the components of the actor selected in the root object view.
-		 * 
-		 * Exists so it can be customized differently depending on whether used in the editor or on the server.
-		 * @note The created view will keep a strong reference to this.
-		 */
-		TSharedPtr<IReplicationSubobjectView> SubobjectView;
-
 		/**
 		 * Optional. Determines the objects displayed as children to the top-level objects in the top section.
 		 * If left unspecified, the top will only display actors.
