Refactor auto adding of objects and properties to be more flexible (not bound to settings) in preparation of Blueprint API

diff --git a/Engine/Plugins/Developer/Concert/ConcertApp/MultiUserClient/Source/MultiUserClient/Private/Replication/Client/ReplicationClient.cpp b/Engine/Plugins/Developer/Concert/ConcertApp/MultiUserClient/Source/MultiUserClient/Private/Replication/Client/ReplicationClient.cpp
--- a/Engine/Plugins/Developer/Concert/ConcertApp/MultiUserClient/Source/MultiUserClient/Private/Replication/Client/ReplicationClient.cpp
+++ b/Engine/Plugins/Developer/Concert/ConcertApp/MultiUserClient/Source/MultiUserClient/Private/Replication/Client/ReplicationClient.cpp
@@ -2,8 +2,10 @@
 
 #include "ReplicationClient.h"
 
+#include "MultiUserReplicationSettings.h"
 #include "Assets/MultiUserReplicationClientPreset.h"
 #include "Replication/Editor/Model/IEditableReplicationStreamModel.h"
+#include "Replication/Editor/Model/StreamExtenderBySettings.h"
 #include "Replication/ReplicationWidgetFactories.h"
 #include "Replication/Stream/StreamChangeTracker.h"
 #include "Replication/Submission/Data/AuthoritySubmission.h"
@@ -28,7 +30,13 @@ namespace UE::MultiUserClient
 		, SubmissionQueue(*SubmissionWorkflow)
 		, LocalClientEditModel(ConcertClientSharedSlate::CreatePropertySelectionModel(
 			*ClientContentStorage->Stream,
-			ClientContentStorage->Stream->MakeReplicationMapGetterAttribute()
+			ClientContentStorage->Stream->MakeReplicationMapGetterAttribute(),
+			// Use MU settings for auto adding properties & objects
+			MakeShared<ConcertClientSharedSlate::FStreamExtenderBySettings>(
+				TAttribute<const FConcertReplicationEditorSettings*>::CreateLambda([]()
+				{
+					return &UMultiUserReplicationSettings::Get()->ReplicationEditorSettings;
+				}))
 			))
 		, LocalClientStreamDiffer(
 			GetStreamSynchronizer(),
diff --git a/Engine/Plugins/Developer/Concert/ConcertApp/MultiUserClient/Source/MultiUserClient/Private/Widgets/ActiveSession/Replication/Client/Multi/SMultiClientView.cpp b/Engine/Plugins/Developer/Concert/ConcertApp/MultiUserClient/Source/MultiUserClient/Private/Widgets/ActiveSession/Replication/Client/Multi/SMultiClientView.cpp
--- a/Engine/Plugins/Developer/Concert/ConcertApp/MultiUserClient/Source/MultiUserClient/Private/Widgets/ActiveSession/Replication/Client/Multi/SMultiClientView.cpp
+++ b/Engine/Plugins/Developer/Concert/ConcertApp/MultiUserClient/Source/MultiUserClient/Private/Widgets/ActiveSession/Replication/Client/Multi/SMultiClientView.cpp
@@ -21,6 +21,12 @@ namespace UE::MultiUserClient
 {
 	void SMultiClientView::Construct(const FArguments& InArgs, TSharedRef<IConcertClient> InConcertClient, FReplicationClientManager& ClientManager, IClientSelectionModel& InDisplayClientsModel)
 	{
+		TAttribute<const FConcertReplicationEditorSettings*> ReplicationSettingsAttribute =
+			TAttribute<const FConcertReplicationEditorSettings*>::CreateLambda([]()
+			{
+				return &UMultiUserReplicationSettings::Get()->ReplicationEditorSettings;
+			});
+		// TODO DP:
 		StreamModel = MakeShared<FMultiStreamModel>(InDisplayClientsModel, ClientManager);
 
 		ChildSlot
@@ -47,21 +53,16 @@ namespace UE::MultiUserClient
 		];
 	}
 
-	TSharedRef<SWidget> SMultiClientView::CreateEditorContent(TSharedRef<IConcertClient> InConcertClient, FReplicationClientManager& InClientManager)
+	TSharedRef<SWidget> SMultiClientView::CreateEditorContent(const TSharedRef<IConcertClient>& InConcertClient, FReplicationClientManager& InClientManager)
 	{
 		using namespace UE::ConcertClientSharedSlate;
 
-		TAttribute<const FConcertReplicationEditorSettings*> ReplicationSettingsAttribute =
-			TAttribute<const FConcertReplicationEditorSettings*>::CreateLambda([]()
-			{
-				return &UMultiUserReplicationSettings::Get()->ReplicationEditorSettings;
-			});
 		TAttribute<TSharedPtr<IMultiReplicationStreamEditor>> MultiStreamEditorAttribute =
 		   TAttribute<TSharedPtr<IMultiReplicationStreamEditor>>::CreateLambda([this]()
 		   {
 			   return StreamEditor;
 		   });
-		TAttribute<IReplicationStreamModel*> ConsolidatedStreamModelAttribute =
+		const TAttribute<IReplicationStreamModel*> ConsolidatedStreamModelAttribute =
 		   TAttribute<IReplicationStreamModel*>::CreateLambda([this]()
 		   {
 			   return &StreamEditor->GetConsolidatedModel();
@@ -72,7 +73,6 @@ namespace UE::MultiUserClient
 			.MultiStreamModel = StreamModel.ToSharedRef(),
 			.ObjectSource = MakeShared<FActorSelectionSourceModel>(),
 			.PropertySource = MakeShared<FSelectPropertyFromUClassModel>(),
-			.ReplicationSettingsAttribute = MoveTemp(ReplicationSettingsAttribute),
 			.ViewerParams =
 			{
 				.SubobjectModel = CreateDefaultComponentHierarchySubobjectModel(), // This makes actors have children in the top view
diff --git a/Engine/Plugins/Developer/Concert/ConcertApp/MultiUserClient/Source/MultiUserClient/Private/Widgets/ActiveSession/Replication/Client/Multi/SMultiClientView.h b/Engine/Plugins/Developer/Concert/ConcertApp/MultiUserClient/Source/MultiUserClient/Private/Widgets/ActiveSession/Replication/Client/Multi/SMultiClientView.h
--- a/Engine/Plugins/Developer/Concert/ConcertApp/MultiUserClient/Source/MultiUserClient/Private/Widgets/ActiveSession/Replication/Client/Multi/SMultiClientView.h
+++ b/Engine/Plugins/Developer/Concert/ConcertApp/MultiUserClient/Source/MultiUserClient/Private/Widgets/ActiveSession/Replication/Client/Multi/SMultiClientView.h
@@ -49,7 +49,7 @@ namespace UE::MultiUserClient
 		TSharedPtr<ConcertClientSharedSlate::IMultiReplicationStreamEditor> StreamEditor;
 
 		/** Creates this widget's editor content */
-		TSharedRef<SWidget> CreateEditorContent(TSharedRef<IConcertClient> InConcertClient, FReplicationClientManager& InClientManager);
+		TSharedRef<SWidget> CreateEditorContent(const TSharedRef<IConcertClient>& InConcertClient, FReplicationClientManager& InClientManager);
 
 		// SClientToolbar attributes
 		TSet<FGuid> GetDisplayClientIds() const;
diff --git a/Engine/Plugins/Developer/Concert/ConcertApp/MultiUserClient/Source/MultiUserClient/Private/Widgets/ActiveSession/Replication/Client/Single/SReplicationClientView.cpp b/Engine/Plugins/Developer/Concert/ConcertApp/MultiUserClient/Source/MultiUserClient/Private/Widgets/ActiveSession/Replication/Client/Single/SReplicationClientView.cpp
--- a/Engine/Plugins/Developer/Concert/ConcertApp/MultiUserClient/Source/MultiUserClient/Private/Widgets/ActiveSession/Replication/Client/Single/SReplicationClientView.cpp
+++ b/Engine/Plugins/Developer/Concert/ConcertApp/MultiUserClient/Source/MultiUserClient/Private/Widgets/ActiveSession/Replication/Client/Single/SReplicationClientView.cpp
@@ -67,8 +67,6 @@ namespace UE::MultiUserClient
 		
 		const TAttribute<const IReplicationStreamViewer*> GetReplicationViewerAttribute =
 			TAttribute<const IReplicationStreamViewer*>::CreateLambda([this](){ return EditorView_TwoSectioned.Get(); });
-		TAttribute<const FConcertReplicationEditorSettings*> ReplicationSettingsAttribute =
-			TAttribute<const FConcertReplicationEditorSettings*>::CreateLambda([](){ return &UMultiUserReplicationSettings::Get()->ReplicationEditorSettings; });
 
 		// Add checkboxes in front of top level and subobject rows for changing authority
 		const FCreateEditorParams ReplicationEditorCreationParams
@@ -78,7 +76,6 @@ namespace UE::MultiUserClient
 			.PropertySource = MakeShared<FSelectPropertyFromUClassModel>(),
 			.IsEditingEnabled = TAttribute<bool>::CreateLambda([&SubmissionWorkflow](){ return CanEverSubmit(SubmissionWorkflow.GetUploadability()); }),
 			.EditingDisabledToolTipText = LOCTEXT("Editing.NotImplemented", "Editing remote clients is not implemented. You can only edit the local client."),
-			.ReplicationSettingsAttribute = MoveTemp(ReplicationSettingsAttribute),
 			.ViewerParams =
 			{
 				.SubobjectModel = CreateDefaultComponentHierarchySubobjectModel(), // This makes actors have children in the top view
diff --git a/Engine/Plugins/Developer/Concert/ConcertUI/ConcertClientSharedSlate/Source/ConcertClientSharedSlate/Private/Replication/Editor/Model/GenericReplicationStreamModel.cpp b/Engine/Plugins/Developer/Concert/ConcertUI/ConcertClientSharedSlate/Source/ConcertClientSharedSlate/Private/Replication/Editor/Model/GenericReplicationStreamModel.cpp
--- a/Engine/Plugins/Developer/Concert/ConcertUI/ConcertClientSharedSlate/Source/ConcertClientSharedSlate/Private/Replication/Editor/Model/GenericReplicationStreamModel.cpp
+++ b/Engine/Plugins/Developer/Concert/ConcertUI/ConcertClientSharedSlate/Source/ConcertClientSharedSlate/Private/Replication/Editor/Model/GenericReplicationStreamModel.cpp
@@ -4,14 +4,19 @@
 
 #include "Replication/PropertyChainUtils.h"
 #include "Replication/Data/ObjectReplicationMap.h"
-#include "Replication/Settings/ConcertReplicationEditorSettings.h"
 
 #include "Algo/AllOf.h"
+#include "Containers/Queue.h"
+#include "Replication/Editor/Model/IStreamExtender.h"
 
 namespace UE::ConcertClientSharedSlate
 {
-	FGenericReplicationStreamModel::FGenericReplicationStreamModel(TAttribute<FObjectReplicationMap*> ReplicationMapAttribute)
-		: ReplicationMapAttribute(MoveTemp(ReplicationMapAttribute))
+	FGenericReplicationStreamModel::FGenericReplicationStreamModel(
+		TAttribute<FObjectReplicationMap*> InReplicationMapAttribute,
+		TSharedPtr<IStreamExtender> InExtender
+		)
+		: ReplicationMapAttribute(MoveTemp(InReplicationMapAttribute))
+		, Extender(MoveTemp(InExtender))
 	{}
 
 	FSoftClassPath FGenericReplicationStreamModel::GetObjectClass(const FSoftObjectPath& Object) const
@@ -99,7 +104,7 @@ namespace UE::ConcertClientSharedSlate
 			return;
 		}
 		
-		TSet<UObject*> ObjectsNotAdded;
+		TArray<UObject*> AddedObjects;
 		for (UObject* Object : Objects)
 		{
 			const FSoftObjectPath ObjectPath = Object;
@@ -107,29 +112,15 @@ namespace UE::ConcertClientSharedSlate
 			{
 				FReplicatedObjectInfo& ObjectInfo = ReplicationMap->ReplicatedObjects.Add(ObjectPath);
 				ObjectInfo.ClassPath = Object->GetClass();
-			}
-			else
-			{
-				ObjectsNotAdded.Add(Object);
+				AddedObjects.AddUnique(Object);
+				
+				ExtendObjects(*ReplicationMap, *Object, AddedObjects);
 			}
 		}
 		
-		if (LIKELY(ObjectsNotAdded.IsEmpty()))
+		if (!AddedObjects.IsEmpty())
 		{
-			OnObjectsChangedDelegate.Broadcast(Objects, {}, EReplicatedObjectChangeReason::ChangedDirectly);
-		}
-		else if (ObjectsNotAdded.Num() < Objects.Num())
-		{
-			// Uncommon case so it is ok if it is suboptimal
-			TArray<UObject*> Added;
-			for (UObject* Object : Objects)
-			{
-				if (!ObjectsNotAdded.Contains(Object))
-				{
-					Added.Add(Object);
-				}
-			}
-			OnObjectsChangedDelegate.Broadcast(Added, {}, EReplicatedObjectChangeReason::ChangedDirectly);
+			OnObjectsChangedDelegate.Broadcast(AddedObjects, {}, EReplicatedObjectChangeReason::ChangedDirectly);
 		}
 	}
 
@@ -258,4 +249,36 @@ namespace UE::ConcertClientSharedSlate
 			OnPropertiesChangedDelegate.Broadcast();
 		}
 	}
+
+	void FGenericReplicationStreamModel::ExtendObjects(FObjectReplicationMap& ReplicationMap, UObject& AddedObject, TArray<UObject*>& ObjectsAddedSoFar)
+	{
+		if (!Extender)
+		{
+			return;
+		}
+
+		TQueue<UObject*> ObjectsToProcess;
+		ObjectsToProcess.Enqueue(&AddedObject);
+
+		UObject* CurrentObject;
+		while (ObjectsToProcess.Dequeue(CurrentObject))
+		{
+			FReplicatedObjectInfo& ObjectInfo = ReplicationMap.ReplicatedObjects[CurrentObject];
+			Extender->ExtendObjectProperties(*CurrentObject, [&ObjectInfo](FConcertPropertyChain PropertyChain)
+			{
+				ObjectInfo.PropertySelection.ReplicatedProperties.AddUnique(MoveTemp(PropertyChain));
+			});
+
+			Extender->AppendAdditionalObjects(*CurrentObject, [&ReplicationMap, &ObjectsAddedSoFar, &ObjectsToProcess](UObject& AdditionalObject)
+			{
+				if (!ObjectsAddedSoFar.Contains(&AdditionalObject))
+				{
+					ObjectsAddedSoFar.AddUnique(&AdditionalObject);
+					ReplicationMap.ReplicatedObjects.FindOrAdd(&AdditionalObject)
+						.ClassPath = AdditionalObject.GetClass();
+					ObjectsToProcess.Enqueue(&AdditionalObject);
+				}
+			});
+		}
+	}
 }
diff --git a/Engine/Plugins/Developer/Concert/ConcertUI/ConcertClientSharedSlate/Source/ConcertClientSharedSlate/Private/Replication/Editor/Model/GenericReplicationStreamModel.h b/Engine/Plugins/Developer/Concert/ConcertUI/ConcertClientSharedSlate/Source/ConcertClientSharedSlate/Private/Replication/Editor/Model/GenericReplicationStreamModel.h
--- a/Engine/Plugins/Developer/Concert/ConcertUI/ConcertClientSharedSlate/Source/ConcertClientSharedSlate/Private/Replication/Editor/Model/GenericReplicationStreamModel.h
+++ b/Engine/Plugins/Developer/Concert/ConcertUI/ConcertClientSharedSlate/Source/ConcertClientSharedSlate/Private/Replication/Editor/Model/GenericReplicationStreamModel.h
@@ -11,13 +11,18 @@ struct FObjectReplicationMap;
 
 namespace UE::ConcertClientSharedSlate
 {
+	class IStreamExtender;
+
 	/** Implements logic for editing a FObjectReplicationMap contained in an UObject. */
 	class FGenericReplicationStreamModel
 		: public IEditableReplicationStreamModel
 	{
 	public:
 		
-		FGenericReplicationStreamModel(TAttribute<FObjectReplicationMap*> ReplicationMapAttribute);
+		FGenericReplicationStreamModel(
+			TAttribute<FObjectReplicationMap*> InReplicationMapAttribute,
+			TSharedPtr<IStreamExtender> InExtender = nullptr
+			);
 		
 		//~ Begin IReplicationStreamModel Interface
 		virtual FSoftClassPath GetObjectClass(const FSoftObjectPath& Object) const override;
@@ -41,8 +46,14 @@ namespace UE::ConcertClientSharedSlate
 		/** Returns the replication map that is supposed to be edited. */
 		TAttribute<FObjectReplicationMap*> ReplicationMapAttribute;
 
+		/** Adds properties and objects when an object is added. Can be null. */
+		TSharedPtr<IStreamExtender> Extender;
+
 		FOnObjectsChanged OnObjectsChangedDelegate;
 		FOnPropertiesChanged OnPropertiesChangedDelegate;
+
+		/** Applies Extender to AddedObject whilst adding any additionally added UObjects to ObjectsAddedSoFar and avoiding adding objects already added to ObjectsAddedSoFar. */
+		void ExtendObjects(FObjectReplicationMap& ReplicationMap, UObject& AddedObject, TArray<UObject*>& ObjectsAddedSoFar);
 	};
 }
 
diff --git a/Engine/Plugins/Developer/Concert/ConcertUI/ConcertClientSharedSlate/Source/ConcertClientSharedSlate/Private/Replication/Editor/Model/StreamExtenderBySettings.cpp b/Engine/Plugins/Developer/Concert/ConcertUI/ConcertClientSharedSlate/Source/ConcertClientSharedSlate/Private/Replication/Editor/Model/StreamExtenderBySettings.cpp
new file mode 100644
--- /dev/null
+++ b/Engine/Plugins/Developer/Concert/ConcertUI/ConcertClientSharedSlate/Source/ConcertClientSharedSlate/Private/Replication/Editor/Model/StreamExtenderBySettings.cpp
@@ -0,0 +1,41 @@
+// Copyright Epic Games, Inc. All Rights Reserved.
+
+#include "Replication/Editor/Model/StreamExtenderBySettings.h"
+
+#include "Replication/Data/ConcertPropertySelection.h"
+#include "Replication/Settings/ConcertReplicationEditorSettings.h"
+
+namespace UE::ConcertClientSharedSlate
+{
+	FStreamExtenderBySettings::FStreamExtenderBySettings(TAttribute<const FConcertReplicationEditorSettings*> InReplicationSettingsAttribute)
+		: ReplicationSettingsAttribute(MoveTemp(InReplicationSettingsAttribute))
+	{}
+
+	void FStreamExtenderBySettings::ExtendObjectProperties(UObject& Object, FEnumerateProperties ForEachPropertyToAdd)
+	{
+		const FConcertReplicationEditorSettings* Settings = ReplicationSettingsAttribute.Get();
+		if (!ensure(Settings))
+		{
+			return;
+		}
+		
+		Settings->AddDefaultPropertiesFromSettings(*Object.GetClass(), [&ForEachPropertyToAdd](FConcertPropertyChain&& PropertyChain)
+		{
+			ForEachPropertyToAdd(MoveTemp(PropertyChain));
+		});
+	}
+
+	void FStreamExtenderBySettings::AppendAdditionalObjects(UObject& Object, FEnumerateObjects ForEachAdditionalObject)
+	{
+		const FConcertReplicationEditorSettings* Settings = ReplicationSettingsAttribute.Get();
+		if (!ensure(Settings))
+		{
+			return;
+		}
+		
+		Settings->AddAdditionalObjectsFromSettings(Object, [&ForEachAdditionalObject](UObject& AdditionalObject)
+		{
+			ForEachAdditionalObject(AdditionalObject);
+		});
+	}
+}
diff --git a/Engine/Plugins/Developer/Concert/ConcertUI/ConcertClientSharedSlate/Source/ConcertClientSharedSlate/Private/Replication/Editor/Model/TransactionalReplicationStreamModel.cpp b/Engine/Plugins/Developer/Concert/ConcertUI/ConcertClientSharedSlate/Source/ConcertClientSharedSlate/Private/Replication/Editor/Model/TransactionalReplicationStreamModel.cpp
--- a/Engine/Plugins/Developer/Concert/ConcertUI/ConcertClientSharedSlate/Source/ConcertClientSharedSlate/Private/Replication/Editor/Model/TransactionalReplicationStreamModel.cpp
+++ b/Engine/Plugins/Developer/Concert/ConcertUI/ConcertClientSharedSlate/Source/ConcertClientSharedSlate/Private/Replication/Editor/Model/TransactionalReplicationStreamModel.cpp
@@ -11,9 +11,10 @@ namespace UE::ConcertClientSharedSlate
 {
 	FTransactionalReplicationStreamModel::FTransactionalReplicationStreamModel(
 		UObject& OwningObject,
-		TAttribute<FObjectReplicationMap*> ReplicationMapAttribute
+		TAttribute<FObjectReplicationMap*> ReplicationMapAttribute,
+		TSharedPtr<IStreamExtender> Extender
 		)
-		: FGenericReplicationStreamModel(MoveTemp(ReplicationMapAttribute))
+		: FGenericReplicationStreamModel(MoveTemp(ReplicationMapAttribute), MoveTemp(Extender))
 		, OwningObject(&OwningObject)
 	{}
 
diff --git a/Engine/Plugins/Developer/Concert/ConcertUI/ConcertClientSharedSlate/Source/ConcertClientSharedSlate/Private/Replication/Editor/Model/TransactionalReplicationStreamModel.h b/Engine/Plugins/Developer/Concert/ConcertUI/ConcertClientSharedSlate/Source/ConcertClientSharedSlate/Private/Replication/Editor/Model/TransactionalReplicationStreamModel.h
--- a/Engine/Plugins/Developer/Concert/ConcertUI/ConcertClientSharedSlate/Source/ConcertClientSharedSlate/Private/Replication/Editor/Model/TransactionalReplicationStreamModel.h
+++ b/Engine/Plugins/Developer/Concert/ConcertUI/ConcertClientSharedSlate/Source/ConcertClientSharedSlate/Private/Replication/Editor/Model/TransactionalReplicationStreamModel.h
@@ -7,6 +7,8 @@
 
 namespace UE::ConcertClientSharedSlate
 {
+	class IStreamExtender;
+
 	/** Special case of FGenericPropertySelectionModel where the edited FObjectReplicationMap lives in an UObject that is RF_Transactional. */
 	class FTransactionalReplicationStreamModel
 		: public FGenericReplicationStreamModel
@@ -16,7 +18,8 @@ namespace UE::ConcertClientSharedSlate
 
 		FTransactionalReplicationStreamModel(
 			UObject& OwningObject,
-			TAttribute<FObjectReplicationMap*> ReplicationMapAttribute
+			TAttribute<FObjectReplicationMap*> ReplicationMapAttribute,
+			TSharedPtr<IStreamExtender> Extender = nullptr
 			);
 		
 		//~ Begin IEditableReplicationStreamModel Interface
diff --git a/Engine/Plugins/Developer/Concert/ConcertUI/ConcertClientSharedSlate/Source/ConcertClientSharedSlate/Private/Replication/Editor/View/ObjectEditor/SBaseReplicationStreamEditor.cpp b/Engine/Plugins/Developer/Concert/ConcertUI/ConcertClientSharedSlate/Source/ConcertClientSharedSlate/Private/Replication/Editor/View/ObjectEditor/SBaseReplicationStreamEditor.cpp
--- a/Engine/Plugins/Developer/Concert/ConcertUI/ConcertClientSharedSlate/Source/ConcertClientSharedSlate/Private/Replication/Editor/View/ObjectEditor/SBaseReplicationStreamEditor.cpp
+++ b/Engine/Plugins/Developer/Concert/ConcertUI/ConcertClientSharedSlate/Source/ConcertClientSharedSlate/Private/Replication/Editor/View/ObjectEditor/SBaseReplicationStreamEditor.cpp
@@ -41,7 +41,6 @@ namespace UE::ConcertClientSharedSlate
 
 		IsEditingEnabledAttribute = InArgs._IsEditingEnabled;
 		EditingDisabledToolTipTextAttribute = InArgs._EditingDisabledToolTipText;
-		ReplicationSettingsAttribute = InArgs._ReplicationSettings;
 		OnExtendObjectsContextMenuDelegate = InArgs._OnExtendObjectsContextMenu;
 		
 		ChildSlot
@@ -116,8 +115,6 @@ namespace UE::ConcertClientSharedSlate
 		// Newly added objects should be automatically selected
 		if (!AddedObjects.IsEmpty())
 		{
-			AutoAddObjectsAndPropertiesFromSettings(AddedObjects);
-			
 			TArray<FSoftObjectPath> TopLevelObjects;
 			Algo::TransformIf(AddedObjects, TopLevelObjects, [this](const UObject* Object)
 			{
@@ -129,39 +126,6 @@ namespace UE::ConcertClientSharedSlate
 		}
 	}
 
-	void SBaseReplicationStreamEditor::AutoAddObjectsAndPropertiesFromSettings(TConstArrayView<UObject*> AddedObjects)
-	{
-		if (!bIsAddingFromSelection)
-		{
-			return;
-		}
-		
-		const FConcertReplicationEditorSettings* AutoPopulateSettings = ReplicationSettingsAttribute.IsBound()
-		  ? ReplicationSettingsAttribute.Get()
-		  : nullptr;
-		if (!AutoPopulateSettings)
-		{
-			return;
-		}
-
-		for (const UObject* AddedObject : AddedObjects)
-		{
-			TArray<FConcertPropertyChain> AdditionalProperties;
-			AutoPopulateSettings->AddDefaultPropertiesFromSettings(*AddedObject->GetClass(), [&AdditionalProperties](FConcertPropertyChain&& Chain)
-			{
-				AdditionalProperties.Emplace(MoveTemp(Chain));
-			});
-			EditablePropertiesModel->AddProperties({ AddedObject }, AdditionalProperties);
-
-			TArray<UObject*> AdditionalObjectsToAdd;
-			AutoPopulateSettings->AddAdditionalObjectsFromSettings(*AddedObject, [&AdditionalObjectsToAdd](UObject& FurtherObject)
-			{
-				AdditionalObjectsToAdd.Add(&FurtherObject);
-			});
-			EditablePropertiesModel->AddObjects(AdditionalObjectsToAdd);
-		}
-	}
-
 	void SBaseReplicationStreamEditor::OnPropertiesChanged()
 	{
 		ReplicationViewer->RefreshPropertyData();
diff --git a/Engine/Plugins/Developer/Concert/ConcertUI/ConcertClientSharedSlate/Source/ConcertClientSharedSlate/Private/Replication/Editor/View/ObjectEditor/SBaseReplicationStreamEditor.h b/Engine/Plugins/Developer/Concert/ConcertUI/ConcertClientSharedSlate/Source/ConcertClientSharedSlate/Private/Replication/Editor/View/ObjectEditor/SBaseReplicationStreamEditor.h
--- a/Engine/Plugins/Developer/Concert/ConcertUI/ConcertClientSharedSlate/Source/ConcertClientSharedSlate/Private/Replication/Editor/View/ObjectEditor/SBaseReplicationStreamEditor.h
+++ b/Engine/Plugins/Developer/Concert/ConcertUI/ConcertClientSharedSlate/Source/ConcertClientSharedSlate/Private/Replication/Editor/View/ObjectEditor/SBaseReplicationStreamEditor.h
@@ -74,9 +74,6 @@ namespace UE::ConcertClientSharedSlate
 			SLATE_ATTRIBUTE(bool, IsEditingEnabled)
 			/** Optional. Whenever IsEditingEnabled returns true, this tooltip is displayed for relevant, disabled UI. */
 			SLATE_ATTRIBUTE(FText, EditingDisabledToolTipText)
-
-			/** Optional settings for auto adding common properties and objects. */
-			SLATE_ATTRIBUTE(const FConcertReplicationEditorSettings*, ReplicationSettings)
 		SLATE_END_ARGS()
 
 		void Construct(const FArguments& InArgs,
@@ -116,8 +113,6 @@ namespace UE::ConcertClientSharedSlate
 		TAttribute<bool> IsEditingEnabledAttribute;
 		/** Optional. Whenever IsEditingEnabled returns true, this tooltip is displayed for relevant, disabled UI. */
 		TAttribute<FText> EditingDisabledToolTipTextAttribute;
-		/** Optional settings for auto adding common properties and objects. */
-		TAttribute<const FConcertReplicationEditorSettings*> ReplicationSettingsAttribute;
 		
 		/** Generates additional entries */
 		FExtendMenu OnExtendObjectsContextMenuDelegate;
@@ -133,7 +128,6 @@ namespace UE::ConcertClientSharedSlate
 		FText GetEditingDisabledText() const;
 		
 		void OnObjectsChanged(TConstArrayView<UObject*> AddedObjects, TConstArrayView<FSoftObjectPath> RemovedObjects, EReplicatedObjectChangeReason ChangeReason);
-		void AutoAddObjectsAndPropertiesFromSettings(TConstArrayView<UObject*> AddedObjects);
 		void OnPropertiesChanged();
 
 		// Customizing left objects search bar
diff --git a/Engine/Plugins/Developer/Concert/ConcertUI/ConcertClientSharedSlate/Source/ConcertClientSharedSlate/Private/Replication/Editor/View/ObjectEditor/SDefaultReplicationStreamEditor.cpp b/Engine/Plugins/Developer/Concert/ConcertUI/ConcertClientSharedSlate/Source/ConcertClientSharedSlate/Private/Replication/Editor/View/ObjectEditor/SDefaultReplicationStreamEditor.cpp
--- a/Engine/Plugins/Developer/Concert/ConcertUI/ConcertClientSharedSlate/Source/ConcertClientSharedSlate/Private/Replication/Editor/View/ObjectEditor/SDefaultReplicationStreamEditor.cpp
+++ b/Engine/Plugins/Developer/Concert/ConcertUI/ConcertClientSharedSlate/Source/ConcertClientSharedSlate/Private/Replication/Editor/View/ObjectEditor/SDefaultReplicationStreamEditor.cpp
@@ -49,7 +49,6 @@ namespace UE::ConcertClientSharedSlate
 				]
 				.IsEditingEnabled(InArgs._IsEditingEnabled)
 				.EditingDisabledToolTipText(InArgs._EditingDisabledToolTipText)
-				.ReplicationSettings(InArgs._ReplicationSettings)
 		];
 	}
 
diff --git a/Engine/Plugins/Developer/Concert/ConcertUI/ConcertClientSharedSlate/Source/ConcertClientSharedSlate/Private/Replication/Editor/View/ObjectEditor/SDefaultReplicationStreamEditor.h b/Engine/Plugins/Developer/Concert/ConcertUI/ConcertClientSharedSlate/Source/ConcertClientSharedSlate/Private/Replication/Editor/View/ObjectEditor/SDefaultReplicationStreamEditor.h
--- a/Engine/Plugins/Developer/Concert/ConcertUI/ConcertClientSharedSlate/Source/ConcertClientSharedSlate/Private/Replication/Editor/View/ObjectEditor/SDefaultReplicationStreamEditor.h
+++ b/Engine/Plugins/Developer/Concert/ConcertUI/ConcertClientSharedSlate/Source/ConcertClientSharedSlate/Private/Replication/Editor/View/ObjectEditor/SDefaultReplicationStreamEditor.h
@@ -58,9 +58,6 @@ namespace UE::ConcertClientSharedSlate
 			SLATE_ATTRIBUTE(bool, IsEditingEnabled)
 			/** Optional. Whenever IsEditingEnabled returns true, this tooltip is displayed for relevant, disabled UI. */
 			SLATE_ATTRIBUTE(FText, EditingDisabledToolTipText)
-
-			/** Optional settings for auto adding common properties and objects. */
-			SLATE_ATTRIBUTE(const FConcertReplicationEditorSettings*, ReplicationSettings)
 		SLATE_END_ARGS()
 
 		void Construct(
diff --git a/Engine/Plugins/Developer/Concert/ConcertUI/ConcertClientSharedSlate/Source/ConcertClientSharedSlate/Private/Replication/ReplicationWidgetFactories.cpp b/Engine/Plugins/Developer/Concert/ConcertUI/ConcertClientSharedSlate/Source/ConcertClientSharedSlate/Private/Replication/ReplicationWidgetFactories.cpp
--- a/Engine/Plugins/Developer/Concert/ConcertUI/ConcertClientSharedSlate/Source/ConcertClientSharedSlate/Private/Replication/ReplicationWidgetFactories.cpp
+++ b/Engine/Plugins/Developer/Concert/ConcertUI/ConcertClientSharedSlate/Source/ConcertClientSharedSlate/Private/Replication/ReplicationWidgetFactories.cpp
@@ -26,8 +26,7 @@ namespace UE::ConcertClientSharedSlate
 			.LeftOfObjectSearchBar() [ Params.ViewerParams.LeftOfObjectSearchBar.Widget ]
 			.LeftOfPropertySearchBar() [ Params.ViewerParams.LeftOfPropertySearchBar.Widget ]
 			.IsEditingEnabled(Params.IsEditingEnabled)
-			.EditingDisabledToolTipText(Params.EditingDisabledToolTipText)
-			.ReplicationSettings(Params.ReplicationSettingsAttribute);
+			.EditingDisabledToolTipText(Params.EditingDisabledToolTipText);
 	}
 
 	TSharedRef<IReplicationStreamEditor> CreateDefaultStreamEditor(FCreateEditorParams Params)
@@ -41,18 +40,19 @@ namespace UE::ConcertClientSharedSlate
 			.LeftOfObjectSearchBar() [ Params.ViewerParams.LeftOfObjectSearchBar.Widget ]
 			.LeftOfPropertySearchBar() [ Params.ViewerParams.LeftOfPropertySearchBar.Widget ]
 			.IsEditingEnabled(Params.IsEditingEnabled)
-			.EditingDisabledToolTipText(Params.EditingDisabledToolTipText)
-			.ReplicationSettings(Params.ReplicationSettingsAttribute);
+			.EditingDisabledToolTipText(Params.EditingDisabledToolTipText);
 	}
 
 	TSharedRef<IEditableReplicationStreamModel> CreatePropertySelectionModel(
 		UObject& OwnerObject,
-		TAttribute<FObjectReplicationMap*> ReplicationMapAttribute
+		TAttribute<FObjectReplicationMap*> ReplicationMapAttribute,
+		TSharedPtr<IStreamExtender> Extender 
 		)
 	{
 		return MakeShared<FTransactionalReplicationStreamModel>(
 			OwnerObject,
-			MoveTemp(ReplicationMapAttribute)
+			MoveTemp(ReplicationMapAttribute),
+			MoveTemp(Extender)
 			);
 	}
 
diff --git a/Engine/Plugins/Developer/Concert/ConcertUI/ConcertClientSharedSlate/Source/ConcertClientSharedSlate/Private/Replication/Settings/ConcertReplicationEditorSettings.cpp b/Engine/Plugins/Developer/Concert/ConcertUI/ConcertClientSharedSlate/Source/ConcertClientSharedSlate/Private/Replication/Settings/ConcertReplicationEditorSettings.cpp
--- a/Engine/Plugins/Developer/Concert/ConcertUI/ConcertClientSharedSlate/Source/ConcertClientSharedSlate/Private/Replication/Settings/ConcertReplicationEditorSettings.cpp
+++ b/Engine/Plugins/Developer/Concert/ConcertUI/ConcertClientSharedSlate/Source/ConcertClientSharedSlate/Private/Replication/Settings/ConcertReplicationEditorSettings.cpp
@@ -3,7 +3,6 @@
 #include "Replication/Settings/ConcertReplicationEditorSettings.h"
 
 #include "ConcertLogGlobal.h"
-#include "Replication/Data/ObjectReplicationMap.h"
 #include "Replication/PropertyChainUtils.h"
 #include "Replication/Settings/ConcertDefaultPropertySelection.h"
 #include "Replication/Settings/ConcertDefaultSubobjectSelection.h"
@@ -163,5 +162,6 @@ namespace UE::ConcertClientSharedSlate::DefaultSubobjects
 
 void FConcertReplicationEditorSettings::AddAdditionalObjectsFromSettings(const UObject& AddedObject, TFunctionRef<void(UObject&)> FurtherObjectsCallback) const
 {
+	checkf(!AddedObject.IsA<UClass>(), TEXT("Pass in the UObject instanced directly, not its class!"));
 	UE::ConcertClientSharedSlate::DefaultSubobjects::InternalAddAdditionalObjectsFromSettings(*AddedObject.GetClass(), *this, AddedObject, FurtherObjectsCallback);
 }
\ No newline at end of file
diff --git a/Engine/Plugins/Developer/Concert/ConcertUI/ConcertClientSharedSlate/Source/ConcertClientSharedSlate/Public/Replication/Editor/Model/IStreamExtender.h b/Engine/Plugins/Developer/Concert/ConcertUI/ConcertClientSharedSlate/Source/ConcertClientSharedSlate/Public/Replication/Editor/Model/IStreamExtender.h
new file mode 100644
--- /dev/null
+++ b/Engine/Plugins/Developer/Concert/ConcertUI/ConcertClientSharedSlate/Source/ConcertClientSharedSlate/Public/Replication/Editor/Model/IStreamExtender.h
@@ -0,0 +1,28 @@
+// Copyright Epic Games, Inc. All Rights Reserved.
+
+#pragma once
+
+#include "Templates/Function.h"
+
+class UObject;
+struct FConcertPropertyChain;
+
+namespace UE::ConcertClientSharedSlate
+{
+	/** When an object is added to IEditableReplicationStreamModel, these functions are called to allow adding additional objects and properties. */
+	class CONCERTCLIENTSHAREDSLATE_API IStreamExtender
+	{
+	public:
+
+		using FEnumerateProperties = TFunctionRef<void(FConcertPropertyChain)>;
+		using FEnumerateObjects = TFunctionRef<void(UObject&)>;
+
+		/** Called when Object has just been added to the model: allows binding some default properties to the object. */
+		virtual void ExtendObjectProperties(UObject& Object, FEnumerateProperties ForEachPropertyToAdd) {}
+
+		/** Called when object has just been added to the model: allows adding some additional object, for which this extender will be invoked recursively. */
+		virtual void AppendAdditionalObjects(UObject& Object, FEnumerateObjects ForEachAdditionalObject) {}
+
+		virtual ~IStreamExtender() = default;
+	};
+}
diff --git a/Engine/Plugins/Developer/Concert/ConcertUI/ConcertClientSharedSlate/Source/ConcertClientSharedSlate/Public/Replication/Editor/Model/StreamExtenderBySettings.h b/Engine/Plugins/Developer/Concert/ConcertUI/ConcertClientSharedSlate/Source/ConcertClientSharedSlate/Public/Replication/Editor/Model/StreamExtenderBySettings.h
new file mode 100644
--- /dev/null
+++ b/Engine/Plugins/Developer/Concert/ConcertUI/ConcertClientSharedSlate/Source/ConcertClientSharedSlate/Public/Replication/Editor/Model/StreamExtenderBySettings.h
@@ -0,0 +1,30 @@
+// Copyright Epic Games, Inc. All Rights Reserved.
+
+#pragma once
+
+#include "IStreamExtender.h"
+#include "Misc/Attribute.h"
+
+class UObject;
+struct FConcertReplicationEditorSettings;
+
+namespace UE::ConcertClientSharedSlate
+{
+	/** Automatically adds objects and properties according to settings specified in some FConcertReplicationEditorSettings. */
+	class CONCERTCLIENTSHAREDSLATE_API FStreamExtenderBySettings : public IStreamExtender
+	{
+	public:
+		
+		FStreamExtenderBySettings(TAttribute<const FConcertReplicationEditorSettings*> InReplicationSettingsAttribute);
+
+		//~ Begin IStreamExtender Interface
+		virtual void ExtendObjectProperties(UObject& Object, FEnumerateProperties ForEachPropertyToAdd) override;
+		virtual void AppendAdditionalObjects(UObject& Object, FEnumerateObjects ForEachAdditionalObject) override;
+		//~ End IStreamExtender Interface
+
+	private:
+		
+		/** Optional settings for auto adding common properties and objects. */
+		TAttribute<const FConcertReplicationEditorSettings*> ReplicationSettingsAttribute;
+	};
+}
diff --git a/Engine/Plugins/Developer/Concert/ConcertUI/ConcertClientSharedSlate/Source/ConcertClientSharedSlate/Public/Replication/ReplicationWidgetFactories.h b/Engine/Plugins/Developer/Concert/ConcertUI/ConcertClientSharedSlate/Source/ConcertClientSharedSlate/Public/Replication/ReplicationWidgetFactories.h
--- a/Engine/Plugins/Developer/Concert/ConcertUI/ConcertClientSharedSlate/Source/ConcertClientSharedSlate/Public/Replication/ReplicationWidgetFactories.h
+++ b/Engine/Plugins/Developer/Concert/ConcertUI/ConcertClientSharedSlate/Source/ConcertClientSharedSlate/Public/Replication/ReplicationWidgetFactories.h
@@ -26,6 +26,7 @@ namespace UE::ConcertClientSharedSlate
 	class IReplicationStreamEditor;
 	class IReplicationStreamViewer;
 	class IReplicationSubobjectView;
+	class IStreamExtender;
 	class ISubobjectModel;
 	class IPropertySelectionSourceModel;
 	
@@ -35,14 +36,16 @@ namespace UE::ConcertClientSharedSlate
 	 * This model edits a FObjectReplicationMap that is assumed to be within the transactional OwnerObject.
 	 * The model will respond to undo & redo by triggering the model's update callbacks.
 	 * 
-	 * @param OwnerObject The object containing the FObjectReplicationMap.
+	 * @param OwnerObject The object containing the FObjectReplicationMap - used for transactions.
 	 * @param ReplicationMapAttribute Getter for extracting the FObjectReplicationMap to edit
+	 * @param Extender Optional callbacks for adding additional properties and objects when an object is added to the model
 	 * 
 	 * @return A model that will edit the FObjectReplicationMap.
 	 */
 	CONCERTCLIENTSHAREDSLATE_API TSharedRef<IEditableReplicationStreamModel> CreatePropertySelectionModel(
 		UObject& OwnerObject,
-		TAttribute<FObjectReplicationMap*> ReplicationMapAttribute
+		TAttribute<FObjectReplicationMap*> ReplicationMapAttribute,
+		TSharedPtr<IStreamExtender> Extender = nullptr
 		);
 
 	/** Builds a similar tree hierarchy as SSubobjectEditor. Reports only components as subobjects. */
@@ -101,9 +104,6 @@ namespace UE::ConcertClientSharedSlate
 		/** Optional. Whenever IsEditingEnabled returns true, this tooltip is displayed for relevant, disabled UI. */
 		TAttribute<FText> EditingDisabledToolTipText;
 
-		/** Optional settings for auto adding common properties and objects. */
-		TAttribute<const FConcertReplicationEditorSettings*> ReplicationSettingsAttribute;
-
 		/** Base params for customizing the viewing part of the editor */
 		FCreateViewerParams ViewerParams;
 	};
@@ -172,9 +172,6 @@ namespace UE::ConcertClientSharedSlate
 		/** These flags modify the behaviour */
 		EMultiStreamEditorFlags Flags = EMultiStreamEditorFlags::Transactional;
 		
-		/** Optional settings for auto adding common properties and objects. */
-		TAttribute<const FConcertReplicationEditorSettings*> ReplicationSettingsAttribute;
-		
 		/** Base params for customizing the viewing part of the editor */
 		FCreateViewerParams ViewerParams;
 	};
