Replace sort with primary and secondary sort

diff --git a/Engine/Plugins/Developer/Concert/ConcertApp/MultiUserClient/Source/MultiUserClient/Private/Widgets/ActiveSession/Replication/Client/Multi/Columns/MultiStreamColumns.h b/Engine/Plugins/Developer/Concert/ConcertApp/MultiUserClient/Source/MultiUserClient/Private/Widgets/ActiveSession/Replication/Client/Multi/Columns/MultiStreamColumns.h
--- a/Engine/Plugins/Developer/Concert/ConcertApp/MultiUserClient/Source/MultiUserClient/Private/Widgets/ActiveSession/Replication/Client/Multi/Columns/MultiStreamColumns.h
+++ b/Engine/Plugins/Developer/Concert/ConcertApp/MultiUserClient/Source/MultiUserClient/Private/Widgets/ActiveSession/Replication/Client/Multi/Columns/MultiStreamColumns.h
@@ -57,14 +57,14 @@ namespace UE::MultiUserClient::MultiStreamColumns
 	 * A combo box which displays all current owners for an object and allows bulk reassigning properties.
 	 *
 	 * @param ConcertClient Used to look up client names
-	 * @param ConsolidatedModelAttribute Used to get child objects
+	 * @param MultiStreamModelAttribute Used to get child objects from the consolidated model and for requesting resorting the column
 	 * @param ReassignmentLogic Performs the act of reassigning
 	 * @param ClientManager Used to access all clients for display in the combo box drop-down
 	 * @param ColumnsSortPriority The order relative to the other columns
 	 */
 	ConcertClientSharedSlate::ReplicationColumns::FReplicationTopLevelObjectColumn ReassignOwnership(
 		TSharedRef<IConcertClient> ConcertClient,
-		TAttribute<ConcertClientSharedSlate::IReplicationStreamModel*> ConsolidatedModelAttribute,
+		TAttribute<TSharedPtr<ConcertClientSharedSlate::IMultiReplicationStreamEditor>> MultiStreamModelAttribute,
 		FReassignObjectPropertiesLogic& ReassignmentLogic,
 		const FReplicationClientManager& ClientManager,
 		int32 ColumnsSortPriority = static_cast<int32>(EColumnSortOrder::ReassignOwnership)
diff --git a/Engine/Plugins/Developer/Concert/ConcertApp/MultiUserClient/Source/MultiUserClient/Private/Widgets/ActiveSession/Replication/Client/Multi/Columns/MultiStreamColumns_AssignProperties.cpp b/Engine/Plugins/Developer/Concert/ConcertApp/MultiUserClient/Source/MultiUserClient/Private/Widgets/ActiveSession/Replication/Client/Multi/Columns/MultiStreamColumns_AssignProperties.cpp
--- a/Engine/Plugins/Developer/Concert/ConcertApp/MultiUserClient/Source/MultiUserClient/Private/Widgets/ActiveSession/Replication/Client/Multi/Columns/MultiStreamColumns_AssignProperties.cpp
+++ b/Engine/Plugins/Developer/Concert/ConcertApp/MultiUserClient/Source/MultiUserClient/Private/Widgets/ActiveSession/Replication/Client/Multi/Columns/MultiStreamColumns_AssignProperties.cpp
@@ -83,7 +83,7 @@ namespace UE::MultiUserClient::MultiStreamColumns
 		using namespace ConcertClientSharedSlate::ReplicationColumns;
 		check(MultiStreamEditor.IsBound() || MultiStreamEditor.IsSet());
 
-		const auto PopuluateSearch = [MultiStreamEditor, &ClientManager, ConcertClient](const FReplicatedPropertyData& Data, TArray<FString>& InOutSearchStrings)
+		const auto PopulateSearch = [MultiStreamEditor, &ClientManager, ConcertClient](const FReplicatedPropertyData& Data, TArray<FString>& InOutSearchStrings)
 		{
 			AssignPropertyColumnUtils::ForEachStreamAssignedTo(*MultiStreamEditor.Get(), Data.GetProperty(),
 				[&ClientManager, &ConcertClient, &InOutSearchStrings](const TSharedRef<IEditableReplicationStreamModel>& Stream)
@@ -97,16 +97,38 @@ namespace UE::MultiUserClient::MultiStreamColumns
 			return SNew(SAssignPropertyComboBox, MultiStreamEditor.Get().ToSharedRef(), ConcertClient, ClientManager)
 				.DisplayedProperty(Args.RowData.GetProperty())
 				.EditedObjects(DisplayedObjects)
-				.HighlightText(Args.HighlightText);
+				.HighlightText(Args.HighlightText)
+				.OnOptionSelected_Lambda([MultiStreamEditor](auto)
+				{
+					if (const TSharedPtr<IMultiReplicationStreamEditor> Editor = MultiStreamEditor.Get())
+					{
+						Editor->GetEditorBase().RequestPropertyColumnResort(AssignPropertyColumnId);
+					}
+				});
+		};
+		const auto IsLessThan = [MultiStreamEditor, ConcertClient, &ClientManager](const FReplicatedPropertyData& Left, const FReplicatedPropertyData& Right)
+		{
+			const TArray<FSoftObjectPath> DisplayedObjects = MultiStreamEditor.Get()->GetEditorBase().GetObjectsBeingPropertyEdited();
+			const TOptional<FString> LeftClientDisplayString = SAssignPropertyComboBox::GetDisplayString(ConcertClient, ClientManager, Left.GetProperty(), DisplayedObjects);
+			const TOptional<FString> RightClientDisplayString = SAssignPropertyComboBox::GetDisplayString(ConcertClient, ClientManager, Right.GetProperty(), DisplayedObjects);
+			
+			if (LeftClientDisplayString && RightClientDisplayString)
+			{
+				return *LeftClientDisplayString < *RightClientDisplayString;
+			}
+			// Our rule: set < unset. This way unassigned appears last.
+			return LeftClientDisplayString.IsSet() && !RightClientDisplayString.IsSet();
 		};
 		
 		return FReplicationPropertyColumn(
 			typename FReplicationPropertyColumn::FArguments()
-				.PopulateSearchItems_Lambda(PopuluateSearch)
+				.PopulateSearchItems_Lambda(PopulateSearch)
 				.GenerateWidgetColumn_Lambda(GenerateWidgetColumn)
+				.IsLessThan_Lambda(IsLessThan)
 				.ColumnSortOrder(ColumnsSortPriority),
 			SHeaderRow::Column(AssignPropertyColumnId)
-				.DefaultLabel(LOCTEXT("Owner", "Owner"))
+				.DefaultLabel(LOCTEXT("Owner.Label", "Assigned Client"))
+				.ToolTipText(LOCTEXT("Owner.ToolTip", "Client that should replicate this property"))
 				.FillSized(FMultiUserReplicationStyle::Get()->GetFloat(TEXT("AllClients.Property.OwnerColumnWidth")))
 			);
 	}
diff --git a/Engine/Plugins/Developer/Concert/ConcertApp/MultiUserClient/Source/MultiUserClient/Private/Widgets/ActiveSession/Replication/Client/Multi/Columns/MultiStreamColumns_ReassignOwnership.cpp b/Engine/Plugins/Developer/Concert/ConcertApp/MultiUserClient/Source/MultiUserClient/Private/Widgets/ActiveSession/Replication/Client/Multi/Columns/MultiStreamColumns_ReassignOwnership.cpp
--- a/Engine/Plugins/Developer/Concert/ConcertApp/MultiUserClient/Source/MultiUserClient/Private/Widgets/ActiveSession/Replication/Client/Multi/Columns/MultiStreamColumns_ReassignOwnership.cpp
+++ b/Engine/Plugins/Developer/Concert/ConcertApp/MultiUserClient/Source/MultiUserClient/Private/Widgets/ActiveSession/Replication/Client/Multi/Columns/MultiStreamColumns_ReassignOwnership.cpp
@@ -5,6 +5,8 @@
 
 #include "MultiUserReplicationStyle.h"
 #include "SReassignObjectComboBox.h"
+#include "Replication/Editor/View/IMultiReplicationStreamEditor.h"
+#include "Replication/Editor/View/IReplicationStreamEditor.h"
 
 #define LOCTEXT_NAMESPACE "ReassignOwnershipColumn"
 
@@ -14,7 +16,7 @@ namespace UE::MultiUserClient::MultiStreamColumns
 	
 	ConcertClientSharedSlate::ReplicationColumns::FReplicationTopLevelObjectColumn ReassignOwnership(
 		TSharedRef<IConcertClient> ConcertClient,
-		TAttribute<ConcertClientSharedSlate::IReplicationStreamModel*> ConsolidatedModelAttribute,
+		TAttribute<TSharedPtr<ConcertClientSharedSlate::IMultiReplicationStreamEditor>> MultiStreamModelAttribute,
 		FReassignObjectPropertiesLogic& ReassignmentLogic,
 		const FReplicationClientManager& ClientManager,
 		const int32 ColumnsSortPriority
@@ -23,14 +25,39 @@ namespace UE::MultiUserClient::MultiStreamColumns
 		using namespace ConcertClientSharedSlate::ReplicationColumns;
 
 		auto MakeWidget =
-			[ConcertClient, ConsolidatedModelAttribute = MoveTemp(ConsolidatedModelAttribute), &ReassignmentLogic, &ClientManager]
+			[ConcertClient, MultiStreamModelAttribute = MoveTemp(MultiStreamModelAttribute), &ReassignmentLogic, &ClientManager]
 			(const FReplicationTopLevelObjectColumn::FBuildArgs& InArgs)
 			{
 				return SNew(SReassignObjectComboBox, ConcertClient, ReassignmentLogic, ClientManager)
 					.ManagedObject(InArgs.RowData.GetObjectPath())
-					.ConsolidatedModel(ConsolidatedModelAttribute)
-					.HighlightText(InArgs.HighlightText);
+					.ConsolidatedModel_Lambda([MultiStreamModelAttribute]()
+					{
+						const TSharedPtr<ConcertClientSharedSlate::IMultiReplicationStreamEditor> Editor = MultiStreamModelAttribute.Get();
+						return Editor
+							? &MultiStreamModelAttribute.Get()->GetConsolidatedModel()
+							: nullptr;
+					})
+					.HighlightText(InArgs.HighlightText)
+					.OnReassignAllOptionClicked_Lambda([MultiStreamModelAttribute](auto)
+					{
+						if (const TSharedPtr<ConcertClientSharedSlate::IMultiReplicationStreamEditor> Model = MultiStreamModelAttribute.Get())
+						{
+							Model->GetEditorBase().RequestObjectColumnResort(ReassignOwnershipColumnId);
+						}
+					});
 			};
+		auto IsLessThan = [ConcertClient, &ReassignmentLogic](const ConcertClientSharedSlate::FReplicatedObjectData& Left, const ConcertClientSharedSlate::FReplicatedObjectData& Right)
+		{
+			const TOptional<FString> LeftClientDisplayString = SReassignObjectComboBox::GetDisplayString(ConcertClient, ReassignmentLogic, Left.GetObjectPath());
+			const TOptional<FString> RightClientDisplayString = SReassignObjectComboBox::GetDisplayString(ConcertClient, ReassignmentLogic, Right.GetObjectPath());
+			
+			if (LeftClientDisplayString && RightClientDisplayString)
+			{
+				return *LeftClientDisplayString < *RightClientDisplayString;
+			}
+			// Our rule: set < unset. This way unassigned appears last.
+			return LeftClientDisplayString.IsSet() && !RightClientDisplayString.IsSet();
+		};
 		
 		return FReplicationTopLevelObjectColumn(
 			FReplicationTopLevelObjectColumn::FArguments()
@@ -39,9 +66,11 @@ namespace UE::MultiUserClient::MultiStreamColumns
 				{
 					SReassignObjectComboBox::PopulateSearchTerms(*ConcertClient->GetCurrentSession(), ReassignmentLogic, ObjectData.GetObjectPath(), InOutSearchStrings);
 				})
+				.IsLessThan_Lambda(MoveTemp(IsLessThan))
 				.ColumnSortOrder(ColumnsSortPriority),
 			SHeaderRow::Column(ReassignOwnershipColumnId)
-				.DefaultLabel(LOCTEXT("Owner", "Owner"))
+				.DefaultLabel(LOCTEXT("Owner.Label", "Assigned Clients"))
+				.ToolTipText(LOCTEXT("Owner.ToolTip", "Clients that have registered properties for an object"))
 				.FillSized(FMultiUserReplicationStyle::Get()->GetFloat(TEXT("AllClients.Object.OwnerColumnWidth")))
 			);
 	}
diff --git a/Engine/Plugins/Developer/Concert/ConcertApp/MultiUserClient/Source/MultiUserClient/Private/Widgets/ActiveSession/Replication/Client/Multi/Columns/MultiStreamColumns_ReplicationToggle.cpp b/Engine/Plugins/Developer/Concert/ConcertApp/MultiUserClient/Source/MultiUserClient/Private/Widgets/ActiveSession/Replication/Client/Multi/Columns/MultiStreamColumns_ReplicationToggle.cpp
--- a/Engine/Plugins/Developer/Concert/ConcertApp/MultiUserClient/Source/MultiUserClient/Private/Widgets/ActiveSession/Replication/Client/Multi/Columns/MultiStreamColumns_ReplicationToggle.cpp
+++ b/Engine/Plugins/Developer/Concert/ConcertApp/MultiUserClient/Source/MultiUserClient/Private/Widgets/ActiveSession/Replication/Client/Multi/Columns/MultiStreamColumns_ReplicationToggle.cpp
@@ -38,6 +38,7 @@ namespace UE::MultiUserClient::MultiStreamColumns
 				.ColumnSortOrder(ColumnsSortPriority),
 			SHeaderRow::Column(ReplicationToggleColumnId)
 				.DefaultLabel(FText::GetEmpty())
+				.ToolTipText(LOCTEXT("Replicates.ToolTip", "Assign properties first.\nControls whether the object should replicate."))
 				.FixedWidth(FMultiUserReplicationStyle::Get()->GetFloat(TEXT("AllClients.Object.ReplicationToggle")))
 			);
 	}
diff --git a/Engine/Plugins/Developer/Concert/ConcertApp/MultiUserClient/Source/MultiUserClient/Private/Widgets/ActiveSession/Replication/Client/Multi/Columns/SAssignPropertyComboBox.cpp b/Engine/Plugins/Developer/Concert/ConcertApp/MultiUserClient/Source/MultiUserClient/Private/Widgets/ActiveSession/Replication/Client/Multi/Columns/SAssignPropertyComboBox.cpp
--- a/Engine/Plugins/Developer/Concert/ConcertApp/MultiUserClient/Source/MultiUserClient/Private/Widgets/ActiveSession/Replication/Client/Multi/Columns/SAssignPropertyComboBox.cpp
+++ b/Engine/Plugins/Developer/Concert/ConcertApp/MultiUserClient/Source/MultiUserClient/Private/Widgets/ActiveSession/Replication/Client/Multi/Columns/SAssignPropertyComboBox.cpp
@@ -22,10 +22,48 @@
 
 namespace UE::MultiUserClient
 {
+	namespace AssignPropertyComboBox
+	{
+		TArray<FGuid> GetDisplayedClients(const FReplicationClientManager& ClientManager, FConcertPropertyChain DisplayedProperty, const TArray<FSoftObjectPath>& EditedObjects)
+		{
+			TArray<FGuid> Clients;
+			ClientManager.ForEachClient([&DisplayedProperty, &EditedObjects, &Clients](const FReplicationClient& Client)
+			{
+				const TMap<FSoftObjectPath, FReplicatedObjectInfo>& ObjectInfoMap = Client.GetStreamSynchronizer().GetServerState().ReplicatedObjects;
+				for (const FSoftObjectPath& ObjectPath : EditedObjects)
+				{
+					if (const FReplicatedObjectInfo* ObjectInfo = ObjectInfoMap.Find(ObjectPath)
+						; ObjectInfo && ObjectInfo->PropertySelection.ReplicatedProperties.Contains(DisplayedProperty))
+					{
+						Clients.Add(Client.GetEndpointId());
+						return EBreakBehavior::Continue;
+					}
+				}
+				return EBreakBehavior::Continue;
+			});
+			return Clients;
+		}
+	}
+	
+	TOptional<FString> SAssignPropertyComboBox::GetDisplayString(
+		const TSharedRef<IConcertClient>& LocalConcertClient,
+		const FReplicationClientManager& ClientManager,
+		FConcertPropertyChain DisplayedProperty,
+		const TArray<FSoftObjectPath>& EditedObjects)
+	{
+		using SWidgetType = ConcertClientSharedSlate::SHorizontalClientList;
+		const TArray<FGuid> Clients = AssignPropertyComboBox::GetDisplayedClients(ClientManager, DisplayedProperty, EditedObjects);
+		return SWidgetType::GetDisplayString(
+			LocalConcertClient.Get(),
+			Clients,
+			SWidgetType::FSortPredicate::CreateStatic(&SWidgetType::SortLocalClientFirstThenAlphabetical, LocalConcertClient)
+			);
+	}
+
 	void SAssignPropertyComboBox::Construct(const FArguments& InArgs,
-        TSharedRef<ConcertClientSharedSlate::IMultiReplicationStreamEditor> InEditor,
-        TSharedRef<IConcertClient> InConcertClient,
-        FReplicationClientManager& InClientManager
+	    TSharedRef<ConcertClientSharedSlate::IMultiReplicationStreamEditor> InEditor,
+	    TSharedRef<IConcertClient> InConcertClient,
+	    FReplicationClientManager& InClientManager
 	)
 	{
 		Editor = InEditor;
@@ -36,6 +74,8 @@ namespace UE::MultiUserClient
 		EditedObjects = InArgs._EditedObjects;
 		HighlightText = InArgs._HighlightText;
 		check(!EditedObjects.IsEmpty());
+
+		OnOptionClickedDelegate = InArgs._OnOptionSelected;
 		
 		ChildSlot
 		[
@@ -58,23 +98,9 @@ namespace UE::MultiUserClient
 	
 	void SAssignPropertyComboBox::RefreshContentBoxContent()
 	{
-		TArray<FGuid> Clients;
-		ClientManager->ForEachClient([this, &Clients](const FReplicationClient& Client)
-		{
-			const TMap<FSoftObjectPath, FReplicatedObjectInfo>& ObjectInfoMap = Client.GetStreamSynchronizer().GetServerState().ReplicatedObjects;
-			for (const FSoftObjectPath& ObjectPath : EditedObjects)
-			{
-				if (const FReplicatedObjectInfo* ObjectInfo = ObjectInfoMap.Find(ObjectPath)
-					; ObjectInfo && ObjectInfo->PropertySelection.ReplicatedProperties.Contains(Property))
-				{
-					Clients.Add(Client.GetEndpointId());
-					return EBreakBehavior::Continue;
-				}
-			}
-			return EBreakBehavior::Continue;
-		});
-
-		ClientListWidget->RefreshList(Clients);
+		ClientListWidget->RefreshList(
+			AssignPropertyComboBox::GetDisplayedClients(*ClientManager, Property, EditedObjects)
+			);
 	}
 
 	TSharedRef<SWidget> SAssignPropertyComboBox::GetMenuContent()
@@ -184,6 +210,8 @@ namespace UE::MultiUserClient
 				}
 			}
 		}
+		
+		OnOptionClickedDelegate.ExecuteIfBound(EndpointId);
 	}
 
 #define SET_REASON(Text) if (Reason) { *Reason = Text; }
diff --git a/Engine/Plugins/Developer/Concert/ConcertApp/MultiUserClient/Source/MultiUserClient/Private/Widgets/ActiveSession/Replication/Client/Multi/Columns/SAssignPropertyComboBox.h b/Engine/Plugins/Developer/Concert/ConcertApp/MultiUserClient/Source/MultiUserClient/Private/Widgets/ActiveSession/Replication/Client/Multi/Columns/SAssignPropertyComboBox.h
--- a/Engine/Plugins/Developer/Concert/ConcertApp/MultiUserClient/Source/MultiUserClient/Private/Widgets/ActiveSession/Replication/Client/Multi/Columns/SAssignPropertyComboBox.h
+++ b/Engine/Plugins/Developer/Concert/ConcertApp/MultiUserClient/Source/MultiUserClient/Private/Widgets/ActiveSession/Replication/Client/Multi/Columns/SAssignPropertyComboBox.h
@@ -34,11 +34,23 @@ namespace UE::MultiUserClient
 	{
 	public:
 
+		DECLARE_DELEGATE_OneParam(FOnOptionClicked, const FGuid& ClientId);
+
+		/** @return The display string this widget would have with the given state. If unset, no clients are displayed in the combobox.*/
+		static TOptional<FString> GetDisplayString(
+			const TSharedRef<IConcertClient>& LocalConcertClient,
+			const FReplicationClientManager& ClientManager,
+			FConcertPropertyChain DisplayedProperty,
+			const TArray<FSoftObjectPath>& EditedObjects
+			);
+		
 		SLATE_BEGIN_ARGS(SAssignPropertyComboBox)
 		{}
 			SLATE_ARGUMENT(FConcertPropertyChain, DisplayedProperty)
 			SLATE_ARGUMENT(TArray<FSoftObjectPath>, EditedObjects)
 			SLATE_ARGUMENT(TSharedPtr<FText>, HighlightText)
+			/** Called when a valid option is selected. */
+			SLATE_EVENT(FOnOptionClicked, OnOptionSelected)
 		SLATE_END_ARGS()
 
 		void Construct(const FArguments& InArgs,
@@ -62,10 +74,12 @@ namespace UE::MultiUserClient
 
 		/** The static menu content (when there is no drop-down). */
 		TSharedPtr<ConcertClientSharedSlate::SHorizontalClientList> ClientListWidget;
-
 		/** Passed to MakeListWidgetDelegate */
 		TSharedPtr<FText> HighlightText;
 
+		/** Called when a valid option is selected. */
+		FOnOptionClicked OnOptionClickedDelegate;
+
 		/** Updates the content of the combo box */
 		void RefreshContentBoxContent();
 		
diff --git a/Engine/Plugins/Developer/Concert/ConcertApp/MultiUserClient/Source/MultiUserClient/Private/Widgets/ActiveSession/Replication/Client/Multi/Columns/SReassignObjectComboBox.cpp b/Engine/Plugins/Developer/Concert/ConcertApp/MultiUserClient/Source/MultiUserClient/Private/Widgets/ActiveSession/Replication/Client/Multi/Columns/SReassignObjectComboBox.cpp
--- a/Engine/Plugins/Developer/Concert/ConcertApp/MultiUserClient/Source/MultiUserClient/Private/Widgets/ActiveSession/Replication/Client/Multi/Columns/SReassignObjectComboBox.cpp
+++ b/Engine/Plugins/Developer/Concert/ConcertApp/MultiUserClient/Source/MultiUserClient/Private/Widgets/ActiveSession/Replication/Client/Multi/Columns/SReassignObjectComboBox.cpp
@@ -21,21 +21,52 @@
 
 namespace UE::MultiUserClient
 {
+	namespace ReassignObjctComboBox
+	{
+		TArray<FGuid> GetDisplayedClients(
+			const FReassignObjectPropertiesLogic& ReassignmentLogic,
+			const FSoftObjectPath& ManagedObject
+			)
+		{
+			TArray<FGuid> ClientsWithOwnership;
+			ReassignmentLogic.EnumerateClientOwnershipState(ManagedObject, [&ClientsWithOwnership](const FGuid& ClientId, FReassignObjectPropertiesLogic::EOwnershipState Ownership)
+			{
+				if (Ownership == FReassignObjectPropertiesLogic::EOwnershipState::HasObjectRegistered)
+				{
+					ClientsWithOwnership.Add(ClientId);
+				}
+				return EBreakBehavior::Continue;
+			});
+			return ClientsWithOwnership;
+		}
+	}
+	
 	void SReassignObjectComboBox::PopulateSearchTerms(
 		const IConcertClientSession& Session,
 		const FReassignObjectPropertiesLogic& ReassignmentLogic,
-		const FSoftObjectPath& ObjectPath,
+		const FSoftObjectPath& ManagedObject,
 		TArray<FString>& InOutSearchTerms
 		)
 	{
-		ReassignmentLogic.EnumerateClientOwnershipState(ObjectPath, [&Session, &InOutSearchTerms](const FGuid& ClientId, FReassignObjectPropertiesLogic::EOwnershipState Ownership)
+		for (const FGuid& ClientId : ReassignObjctComboBox::GetDisplayedClients(ReassignmentLogic, ManagedObject))
 		{
-			if (Ownership == FReassignObjectPropertiesLogic::EOwnershipState::HasObjectRegistered)
-			{
-				InOutSearchTerms.Add(ClientUtils::GetClientDisplayName(Session, ClientId));
-			}
-			return EBreakBehavior::Continue;
-		});
+			InOutSearchTerms.Add(ClientUtils::GetClientDisplayName(Session, ClientId));
+		}
+	}
+
+	TOptional<FString> SReassignObjectComboBox::GetDisplayString(
+		const TSharedRef<IConcertClient>& LocalConcertClient,
+		const FReassignObjectPropertiesLogic& ReassignmentLogic,
+		const FSoftObjectPath& ManagedObject
+		)
+	{
+		using SWidgetType = ConcertClientSharedSlate::SHorizontalClientList;
+		const TArray<FGuid> Clients = ReassignObjctComboBox::GetDisplayedClients(ReassignmentLogic, ManagedObject);
+		return SWidgetType::GetDisplayString(
+			LocalConcertClient.Get(),
+			Clients,
+			SWidgetType::FSortPredicate::CreateStatic(&SWidgetType::SortLocalClientFirstThenAlphabetical, LocalConcertClient)
+			);
 	}
 
 	void SReassignObjectComboBox::Construct(
@@ -53,6 +84,8 @@ namespace UE::MultiUserClient
 		ManagedObject = InArgs._ManagedObject;
 		ConsolidatedStreamModelAttribute = InArgs._ConsolidatedModel;
 		check(ConsolidatedStreamModelAttribute.IsSet() || ConsolidatedStreamModelAttribute.IsBound());
+
+		OnReassignAllOptionClickedDelegate = InArgs._OnReassignAllOptionClicked;
 		
 		ChildSlot
 		[
@@ -92,19 +125,9 @@ namespace UE::MultiUserClient
 
 	void SReassignObjectComboBox::UpdateComboButtonContent() const
 	{
-		using EOwnership = FReassignObjectPropertiesLogic::EOwnershipState;
-		
-		TArray<FGuid> ClientsWithOwnership;
-		ReassignmentLogic->EnumerateClientOwnershipState(ManagedObject, [&ClientsWithOwnership](const FGuid& ClientId, EOwnership Ownership)
-		{
-			if (Ownership == FReassignObjectPropertiesLogic::EOwnershipState::HasObjectRegistered)
-			{
-				ClientsWithOwnership.Add(ClientId);
-			}
-			return EBreakBehavior::Continue;
-		});
-
-		ComboClientList->RefreshList(ClientsWithOwnership);
+		ComboClientList->RefreshList(
+			ReassignObjctComboBox::GetDisplayedClients(*ReassignmentLogic, ManagedObject)
+			);
 	}
 
 	EVisibility SReassignObjectComboBox::GetThrobberVisibility() const
@@ -167,6 +190,8 @@ namespace UE::MultiUserClient
 						TArray<FSoftObjectPath> ObjectsNotInline;
 						Algo::Transform(ObjectsToAssign.Get(), ObjectsNotInline, [](const FSoftObjectPath& ObjectPath){ return ObjectPath; });
 						ReassignmentLogic->ReassignAllTo(ObjectsNotInline, ClientId);
+						
+						OnReassignAllOptionClickedDelegate.ExecuteIfBound(ClientId);
 					}),
 					FCanExecuteAction::CreateLambda([this, ObjectsToAssign, ClientId](){ return ReassignmentLogic->CanReassignAnyTo(ObjectsToAssign.Get(), ClientId); }),
 					FIsActionChecked::CreateLambda([this, ObjectsToAssign, ClientId](){ return ReassignmentLogic->OwnsAnyOf(ObjectsToAssign.Get(), ClientId); })
diff --git a/Engine/Plugins/Developer/Concert/ConcertApp/MultiUserClient/Source/MultiUserClient/Private/Widgets/ActiveSession/Replication/Client/Multi/Columns/SReassignObjectComboBox.h b/Engine/Plugins/Developer/Concert/ConcertApp/MultiUserClient/Source/MultiUserClient/Private/Widgets/ActiveSession/Replication/Client/Multi/Columns/SReassignObjectComboBox.h
--- a/Engine/Plugins/Developer/Concert/ConcertApp/MultiUserClient/Source/MultiUserClient/Private/Widgets/ActiveSession/Replication/Client/Multi/Columns/SReassignObjectComboBox.h
+++ b/Engine/Plugins/Developer/Concert/ConcertApp/MultiUserClient/Source/MultiUserClient/Private/Widgets/ActiveSession/Replication/Client/Multi/Columns/SReassignObjectComboBox.h
@@ -26,8 +26,22 @@ namespace UE::MultiUserClient
 	{
 	public:
 
+		DECLARE_DELEGATE_OneParam(FOnReassignAllOptionClicked, const FGuid& ClientId);
+
 		/** Fills in the search terms based what the a widget instance would be displaying given its state. */
-		static void PopulateSearchTerms(const IConcertClientSession& Session, const FReassignObjectPropertiesLogic& ReassignmentLogic, const FSoftObjectPath& ObjectPath, TArray<FString>& InOutSearchTerms);
+		static void PopulateSearchTerms(
+			const IConcertClientSession& Session,
+			const FReassignObjectPropertiesLogic& ReassignmentLogic,
+			const FSoftObjectPath& ManagedObject,
+			TArray<FString>& InOutSearchTerms
+			);
+		
+		/** @return The display string this widget would have with the given state. If unset, no clients are displayed in the combobox. */
+		static TOptional<FString> GetDisplayString(
+			const TSharedRef<IConcertClient>& LocalConcertClient,
+			const FReassignObjectPropertiesLogic& ReassignmentLogic,
+			const FSoftObjectPath& ManagedObject
+			);
 
 		SLATE_BEGIN_ARGS(SReassignObjectComboBox)
 		{}
@@ -36,6 +50,9 @@ namespace UE::MultiUserClient
 			SLATE_ARGUMENT(TSharedPtr<FText>, HighlightText)
 			/** Used to to figure out child objects. */
 			SLATE_ATTRIBUTE(ConcertClientSharedSlate::IReplicationStreamModel*, ConsolidatedModel)
+
+			/** Called when a valid client ID is selected for reassignment. */
+			SLATE_EVENT(FOnReassignAllOptionClicked, OnReassignAllOptionClicked)
 		SLATE_END_ARGS()
 
 		void Construct(
@@ -65,6 +82,9 @@ namespace UE::MultiUserClient
 		TSharedPtr<ConcertClientSharedSlate::SHorizontalClientList> ComboClientList;
 		/** Passed to client list for search highlighting */
 		TSharedPtr<FText> HighlightText;
+		
+		/** Called when a valid client ID is selected for reassignment. */
+		FOnReassignAllOptionClicked OnReassignAllOptionClickedDelegate;
 
 		void UpdateComboButtonContent() const;
 
diff --git a/Engine/Plugins/Developer/Concert/ConcertApp/MultiUserClient/Source/MultiUserClient/Private/Widgets/ActiveSession/Replication/Client/Multi/Columns/SReplicationMultiToggleCheckbox.cpp b/Engine/Plugins/Developer/Concert/ConcertApp/MultiUserClient/Source/MultiUserClient/Private/Widgets/ActiveSession/Replication/Client/Multi/Columns/SReplicationMultiToggleCheckbox.cpp
--- a/Engine/Plugins/Developer/Concert/ConcertApp/MultiUserClient/Source/MultiUserClient/Private/Widgets/ActiveSession/Replication/Client/Multi/Columns/SReplicationMultiToggleCheckbox.cpp
+++ b/Engine/Plugins/Developer/Concert/ConcertApp/MultiUserClient/Source/MultiUserClient/Private/Widgets/ActiveSession/Replication/Client/Multi/Columns/SReplicationMultiToggleCheckbox.cpp
@@ -75,7 +75,20 @@ namespace UE::MultiUserClient
 
 	FText SReplicationMultiToggleCheckbox::GetRootToolTipText() const
 	{
-		return LOCTEXT("RootTooltip", "Controls whether the assigned clients should replicate this object.");
+		const bool bHasProperties = IsCheckboxEnabledForObject(Object);
+		if (!bHasProperties)
+		{
+			// TODO UE-200496 Update text if we predict there to be a conflict
+			return LOCTEXT("Toggle.ToolTip.NoProperties", "Assign properties first.");
+		}
+		
+		switch (GetCheckboxStateForObject(Object))
+		{
+		case ECheckBoxState::Unchecked: return LOCTEXT("Toggle.ToolTip.Unchecked", "Not replicating.");
+		case ECheckBoxState::Checked: return LOCTEXT("Toggle.ToolTip.Checked", "Replicating all assigned objects.");
+		case ECheckBoxState::Undetermined: return LOCTEXT("Toggle.ToolTip.Undetermined", "Replicating some assigned objects, but not all.");
+		default: return FText::GetEmpty();
+		}
 	}
 
 	ECheckBoxState SReplicationMultiToggleCheckbox::GetCheckboxStateForObject(FSoftObjectPath InObject) const
diff --git a/Engine/Plugins/Developer/Concert/ConcertApp/MultiUserClient/Source/MultiUserClient/Private/Widgets/ActiveSession/Replication/Client/Multi/SMultiClientView.cpp b/Engine/Plugins/Developer/Concert/ConcertApp/MultiUserClient/Source/MultiUserClient/Private/Widgets/ActiveSession/Replication/Client/Multi/SMultiClientView.cpp
--- a/Engine/Plugins/Developer/Concert/ConcertApp/MultiUserClient/Source/MultiUserClient/Private/Widgets/ActiveSession/Replication/Client/Multi/SMultiClientView.cpp
+++ b/Engine/Plugins/Developer/Concert/ConcertApp/MultiUserClient/Source/MultiUserClient/Private/Widgets/ActiveSession/Replication/Client/Multi/SMultiClientView.cpp
@@ -2,6 +2,7 @@
 
 #include "SMultiClientView.h"
 
+#include "IClientSelectionModel.h"
 #include "MultiStreamModel.h"
 #include "MultiUserReplicationSettings.h"
 #include "Replication/Client/ReplicationClient.h"
@@ -10,6 +11,7 @@
 #include "Replication/Editor/Model/Property/SelectPropertyFromUClassModel.h"
 #include "Replication/Editor/View/IMultiReplicationStreamEditor.h"
 #include "Replication/ReplicationWidgetFactories.h"
+#include "Replication/Editor/View/IReplicationStreamEditor.h"
 #include "Widgets/ActiveSession/Replication/Client/Multi/Columns/MultiStreamColumns.h"
 #include "Widgets/ActiveSession/Replication/Client/SClientToolbar.h"
 
@@ -19,10 +21,15 @@
 
 namespace UE::MultiUserClient
 {
-	void SMultiClientView::Construct(const FArguments& InArgs, TSharedRef<IConcertClient> InConcertClient, FReplicationClientManager& ClientManager, IClientSelectionModel& InDisplayClientsModel)
+	void SMultiClientView::Construct(const FArguments& InArgs, TSharedRef<IConcertClient> InConcertClient, FReplicationClientManager& InClientManager, IClientSelectionModel& InDisplayClientsModel)
 	{
-		StreamModel = MakeShared<FMultiStreamModel>(InDisplayClientsModel, ClientManager);
+		StreamModel = MakeShared<FMultiStreamModel>(InDisplayClientsModel, InClientManager);
 
+		ClientManager = &InClientManager;
+		ClientManager->OnRemoteClientsChanged().AddSP(this, &SMultiClientView::RebuildClientSubscriptions);
+		SelectionModel = &InDisplayClientsModel;
+		SelectionModel->OnSelectionChanged().AddSP(this, &SMultiClientView::RebuildClientSubscriptions);
+		
 		ChildSlot
 		[
 			SNew(SVerticalBox)
@@ -32,7 +39,7 @@ namespace UE::MultiUserClient
 			.AutoHeight()
 			.Padding(2.f)
 			[
-				SAssignNew(Toolbar, SClientToolbar, ClientManager.GetAuthorityCache())
+				SAssignNew(Toolbar, SClientToolbar, InClientManager.GetAuthorityCache())
 				.ViewSelectionArea() [ InArgs._ViewSelectionArea.Widget ]
 				.DisplayedClients(this, &SMultiClientView::GetDisplayClientIds)
 				.ForEachReplicatedObject(this, &SMultiClientView::EnumerateObjectsInStreams)
@@ -42,9 +49,18 @@ namespace UE::MultiUserClient
 			+SVerticalBox::Slot()
 			.FillHeight(1.f)
 			[
-				CreateEditorContent(InConcertClient, ClientManager)
+				CreateEditorContent(InConcertClient, InClientManager)
 			]
 		];
+
+		RebuildClientSubscriptions();
+	}
+
+	SMultiClientView::~SMultiClientView()
+	{
+		ClientManager->OnRemoteClientsChanged().RemoveAll(this);
+		SelectionModel->OnSelectionChanged().RemoveAll(this);
+		CleanClientSubscriptions();
 	}
 
 	TSharedRef<SWidget> SMultiClientView::CreateEditorContent(const TSharedRef<IConcertClient>& InConcertClient, FReplicationClientManager& InClientManager)
@@ -79,9 +95,10 @@ namespace UE::MultiUserClient
 				.AdditionalObjectColumns =
 				{
 					MultiStreamColumns::ReplicationToggle(InConcertClient, ConsolidatedStreamModelAttribute, InClientManager),
-					MultiStreamColumns::ReassignOwnership(InConcertClient, ConsolidatedStreamModelAttribute, InClientManager.GetReassignmentLogic(), InClientManager)
+					MultiStreamColumns::ReassignOwnership(InConcertClient, MultiStreamEditorAttribute, InClientManager.GetReassignmentLogic(), InClientManager)
 				},
-				.AdditionalPropertyColumns = { MultiStreamColumns::AssignPropertyColumn(MoveTemp(MultiStreamEditorAttribute), InConcertClient, InClientManager) }
+				.AdditionalPropertyColumns = { MultiStreamColumns::AssignPropertyColumn(MultiStreamEditorAttribute, InConcertClient, InClientManager) },
+				.PrimaryPropertySort = { MultiStreamColumns::AssignPropertyColumnId, EColumnSortMode::Ascending}
 			}
 		};
 		StreamEditor = CreateBaseMultiStreamEditor(MoveTemp(Params));
@@ -112,6 +129,33 @@ namespace UE::MultiUserClient
 			return EBreakBehavior::Continue;
 		});
 	}
+
+	void SMultiClientView::RebuildClientSubscriptions()
+	{
+		CleanClientSubscriptions();
+
+		ClientManager->ForEachClient([this](FReplicationClient& Client)
+		{
+			Client.OnModelChanged().AddSP(this, &SMultiClientView::OnClientChanged, Client.GetEndpointId());
+			return EBreakBehavior::Continue;
+		});
+	}
+
+	void SMultiClientView::CleanClientSubscriptions()
+	{
+		ClientManager->ForEachClient([this](FReplicationClient& Client)
+		{
+			Client.OnModelChanged().RemoveAll(this);
+			return EBreakBehavior::Continue;
+		});
+	}
+
+	void SMultiClientView::OnClientChanged(FGuid)
+	{
+		// When reassignment operations complete, the content of the columns changes so a resort is required.
+		StreamEditor->GetEditorBase().RequestObjectColumnResort(MultiStreamColumns::ReassignOwnershipColumnId);
+		StreamEditor->GetEditorBase().RequestPropertyColumnResort(MultiStreamColumns::AssignPropertyColumnId);
+	}
 }
 
 #undef LOCTEXT_NAMESPACE
\ No newline at end of file
diff --git a/Engine/Plugins/Developer/Concert/ConcertApp/MultiUserClient/Source/MultiUserClient/Private/Widgets/ActiveSession/Replication/Client/Single/Columns/SingleClientColumns_Owner.cpp b/Engine/Plugins/Developer/Concert/ConcertApp/MultiUserClient/Source/MultiUserClient/Private/Widgets/ActiveSession/Replication/Client/Single/Columns/SingleClientColumns_Owner.cpp
--- a/Engine/Plugins/Developer/Concert/ConcertApp/MultiUserClient/Source/MultiUserClient/Private/Widgets/ActiveSession/Replication/Client/Single/Columns/SingleClientColumns_Owner.cpp
+++ b/Engine/Plugins/Developer/Concert/ConcertApp/MultiUserClient/Source/MultiUserClient/Private/Widgets/ActiveSession/Replication/Client/Single/Columns/SingleClientColumns_Owner.cpp
@@ -54,8 +54,9 @@ namespace UE::MultiUserClient::SingleClientColumns
 						);
 				})
 				.ColumnSortOrder(static_cast<int32>(ETopLevelObjectColumnOrder::Owner)),
-			SHeaderRow::Column(OwnerOfSubobjectColumnId)
-				.DefaultLabel(LOCTEXT("Subobject.Owner", "Owner"))
+				SHeaderRow::Column(OwnerOfSubobjectColumnId)
+				.DefaultLabel(LOCTEXT("Owner.Label", "Assigned Clients"))
+				.ToolTipText(LOCTEXT("Owner.ToolTip", "Clients that have registered properties for an object"))
 				.FillSized(FMultiUserReplicationStyle::Get()->GetFloat(TEXT("SingleClient.Object.OwnerColumnWidth")))
 			);
 	}
diff --git a/Engine/Plugins/Developer/Concert/ConcertApp/MultiUserClient/Source/MultiUserClient/Private/Widgets/ActiveSession/Replication/Misc/SNoClients.h b/Engine/Plugins/Developer/Concert/ConcertApp/MultiUserClient/Source/MultiUserClient/Private/Widgets/ActiveSession/Replication/Misc/SNoClients.h
--- a/Engine/Plugins/Developer/Concert/ConcertApp/MultiUserClient/Source/MultiUserClient/Private/Widgets/ActiveSession/Replication/Misc/SNoClients.h
+++ b/Engine/Plugins/Developer/Concert/ConcertApp/MultiUserClient/Source/MultiUserClient/Private/Widgets/ActiveSession/Replication/Misc/SNoClients.h
@@ -15,16 +15,18 @@ namespace UE::MultiUserClient
 	{
 	public:
 
+		static FText GetDisplayText() { return LOCTEXT("None", " - "); }
+		
 		SLATE_BEGIN_ARGS(SNoClients)
 		{}
 		SLATE_END_ARGS()
-
+		
 		void Construct(const FArguments& InArgs)
 		{
 			ChildSlot
 			[
 				SNew(STextBlock)
-				.Text(LOCTEXT("None", " - "))
+				.Text(GetDisplayText())
 			];
 		}
 	};
diff --git a/Engine/Plugins/Developer/Concert/ConcertUI/ConcertClientSharedSlate/Source/ConcertClientSharedSlate/Private/Replication/Editor/View/ObjectEditor/SBaseReplicationStreamEditor.cpp b/Engine/Plugins/Developer/Concert/ConcertUI/ConcertClientSharedSlate/Source/ConcertClientSharedSlate/Private/Replication/Editor/View/ObjectEditor/SBaseReplicationStreamEditor.cpp
--- a/Engine/Plugins/Developer/Concert/ConcertUI/ConcertClientSharedSlate/Source/ConcertClientSharedSlate/Private/Replication/Editor/View/ObjectEditor/SBaseReplicationStreamEditor.cpp
+++ b/Engine/Plugins/Developer/Concert/ConcertUI/ConcertClientSharedSlate/Source/ConcertClientSharedSlate/Private/Replication/Editor/View/ObjectEditor/SBaseReplicationStreamEditor.cpp
@@ -47,11 +47,14 @@ namespace UE::ConcertClientSharedSlate
 		[
 			SAssignNew(ReplicationViewer, SReplicationStreamViewer, PropertiesModelAdapter.ToSharedRef())
 				.AdditionalObjectColumns(InArgs._AdditionalObjectColumns)
+				.PrimaryObjectSort(InArgs._PrimaryObjectSort)
+				.SecondaryObjectSort(InArgs._SecondaryObjectSort)
 				.AdditionalPropertyColumns(InArgs._AdditionalPropertyColumns)
+				.PrimaryPropertySort(InArgs._PrimaryPropertySort)
+				.SecondaryPropertySort(InArgs._SecondaryPropertySort)
 				.SubobjectModel(InArgs._SubobjectModel)
 				.OnDeleteObjects(this, &SBaseReplicationStreamEditor::OnDeleteObjects)
 				.OnObjectsContextMenuOpening(this, &SBaseReplicationStreamEditor::OnObjectsContextMenuOpening)
-				.SortPropertyRowPredicate(InArgs._SortPropertyRowPredicate)
 				.LeftOfObjectSearchBar()
 				[
 					SNew(SHorizontalBox)
@@ -85,9 +88,14 @@ namespace UE::ConcertClientSharedSlate
 		ReplicationViewer->Refresh();
 	}
 
-	TArray<FSoftObjectPath> SBaseReplicationStreamEditor::GetSelectedTopLevelObjects() const
+	void SBaseReplicationStreamEditor::RequestObjectColumnResort(const FName& ColumnId)
 	{
-		return ReplicationViewer->GetSelectedTopLevelObjects();
+		ReplicationViewer->RequestObjectColumnResort(ColumnId);
+	}
+
+	void SBaseReplicationStreamEditor::RequestPropertyColumnResort(const FName& ColumnId)
+	{
+		ReplicationViewer->RequestPropertyColumnResort(ColumnId);
 	}
 
 	TArray<FSoftObjectPath> SBaseReplicationStreamEditor::GetObjectsBeingPropertyEdited() const
diff --git a/Engine/Plugins/Developer/Concert/ConcertUI/ConcertClientSharedSlate/Source/ConcertClientSharedSlate/Private/Replication/Editor/View/ObjectEditor/SBaseReplicationStreamEditor.h b/Engine/Plugins/Developer/Concert/ConcertUI/ConcertClientSharedSlate/Source/ConcertClientSharedSlate/Private/Replication/Editor/View/ObjectEditor/SBaseReplicationStreamEditor.h
--- a/Engine/Plugins/Developer/Concert/ConcertUI/ConcertClientSharedSlate/Source/ConcertClientSharedSlate/Private/Replication/Editor/View/ObjectEditor/SBaseReplicationStreamEditor.h
+++ b/Engine/Plugins/Developer/Concert/ConcertUI/ConcertClientSharedSlate/Source/ConcertClientSharedSlate/Private/Replication/Editor/View/ObjectEditor/SBaseReplicationStreamEditor.h
@@ -5,6 +5,7 @@
 #include "Model/Item/SourceModelBuilders.h"
 #include "Replication/Editor/View/IReplicationStreamEditor.h"
 #include "Replication/Editor/View/PredefinedReplicationColumns.h"
+#include "Replication/Editor/View/ReplicationColumnsUtils.h"
 #include "Replication/ReplicationWidgetDelegates.h"
 
 #include "Widgets/DeclarativeSyntaxSupport.h"
@@ -53,8 +54,17 @@ namespace UE::ConcertClientSharedSlate
 		{}
 			/** Additional columns to add to the object view */
 			SLATE_ARGUMENT(TArray<ReplicationColumns::FReplicationTopLevelObjectColumn>, AdditionalObjectColumns)
+			/** Initial primary sort to set. */
+			SLATE_ARGUMENT(FColumnSortInfo, PrimaryObjectSort)
+			/** Initial secondary sort to set. */
+			SLATE_ARGUMENT(FColumnSortInfo, SecondaryObjectSort)
+		
 			/** Additional columns to add to the property view */
 			SLATE_ARGUMENT(TArray<ReplicationColumns::FReplicationPropertyColumn>, AdditionalPropertyColumns)
+			/** Initial primary sort to set. */
+			SLATE_ARGUMENT(FColumnSortInfo, PrimaryPropertySort)
+			/** Initial secondary sort to set. */
+			SLATE_ARGUMENT(FColumnSortInfo, SecondaryPropertySort)
 		
 			/** Optional. If set, this determines the children nested under the root objects. */
 			SLATE_ARGUMENT(TSharedPtr<ISubobjectModel>, SubobjectModel)
@@ -62,9 +72,6 @@ namespace UE::ConcertClientSharedSlate
 			/** Called to generate the context menu for objects. This extends the options already generated by this widget. */
 			SLATE_EVENT(FExtendMenu, OnExtendObjectsContextMenu)
 
-			/** Optional. Used for determining the order in which properties are displayed. */
-			SLATE_EVENT(FSortPropertyPredicate, SortPropertyRowPredicate)
-
 			/** Optional widget to add to the left of the object list search bar. */
 			SLATE_NAMED_SLOT(FArguments, LeftOfObjectSearchBar)
 			/** Optional widget to add to the left of the property list search bar. */
@@ -85,7 +92,8 @@ namespace UE::ConcertClientSharedSlate
 
 		//~ Begin IReplicationStreamEditor Interface
 		virtual void Refresh() override;
-		virtual TArray<FSoftObjectPath> GetSelectedTopLevelObjects() const override;
+		virtual void RequestObjectColumnResort(const FName& ColumnId) override;
+		virtual void RequestPropertyColumnResort(const FName& ColumnId) override;
 		virtual TArray<FSoftObjectPath> GetObjectsBeingPropertyEdited() const override;
 		//~ End IReplicationStreamEditor Interface
 
diff --git a/Engine/Plugins/Developer/Concert/ConcertUI/ConcertClientSharedSlate/Source/ConcertClientSharedSlate/Private/Replication/Editor/View/ObjectEditor/SDefaultReplicationStreamEditor.cpp b/Engine/Plugins/Developer/Concert/ConcertUI/ConcertClientSharedSlate/Source/ConcertClientSharedSlate/Private/Replication/Editor/View/ObjectEditor/SDefaultReplicationStreamEditor.cpp
--- a/Engine/Plugins/Developer/Concert/ConcertUI/ConcertClientSharedSlate/Source/ConcertClientSharedSlate/Private/Replication/Editor/View/ObjectEditor/SDefaultReplicationStreamEditor.cpp
+++ b/Engine/Plugins/Developer/Concert/ConcertUI/ConcertClientSharedSlate/Source/ConcertClientSharedSlate/Private/Replication/Editor/View/ObjectEditor/SDefaultReplicationStreamEditor.cpp
@@ -30,15 +30,26 @@ namespace UE::ConcertClientSharedSlate
 			);
 		TArray<FReplicationPropertyColumn> PropertyColumns = InArgs._AdditionalPropertyColumns;
 		PropertyColumns.Add(ReplicatesColumn);
+
+		// Set both primary and secondary in case one is overriden but always use the override.
+		const FColumnSortInfo PrimaryPropertySort = InArgs._PrimaryObjectSort.IsValid()
+			? InArgs._PrimaryObjectSort
+			: FColumnSortInfo{ ReplicatesColumnId, EColumnSortMode::Ascending };
+		const FColumnSortInfo SecondaryPropertySort = InArgs._SecondaryObjectSort.IsValid()
+			? InArgs._SecondaryObjectSort
+			: FColumnSortInfo{ LabelColumnId, EColumnSortMode::Ascending };
 		
 		ChildSlot
 		[
 			SAssignNew(WrappedEditor, SBaseReplicationStreamEditor, InPropertiesModel, InObjectSelectionSource, InPropertySelectionSource)
 				.AdditionalObjectColumns(InArgs._AdditionalObjectColumns)
+				.PrimaryObjectSort(InArgs._PrimaryObjectSort)
+				.SecondaryObjectSort(InArgs._SecondaryObjectSort)
 				.AdditionalPropertyColumns(PropertyColumns)
+				.PrimaryPropertySort(PrimaryPropertySort)
+				.SecondaryPropertySort(SecondaryPropertySort)
 				.SubobjectModel(InArgs._SubobjectModel)
 				.OnExtendObjectsContextMenu(InArgs._OnExtendObjectsContextMenu)
-				.SortPropertyRowPredicate(this, &SDefaultReplicationStreamEditor::SortPropertiesPredicate)
 				.LeftOfObjectSearchBar()
 				[
 					InArgs._LeftOfObjectSearchBar.Widget
@@ -57,26 +68,18 @@ namespace UE::ConcertClientSharedSlate
 		return WrappedEditor->Refresh();
 	}
 
-	TArray<FSoftObjectPath> SDefaultReplicationStreamEditor::GetSelectedTopLevelObjects() const
+	void SDefaultReplicationStreamEditor::RequestObjectColumnResort(const FName& ColumnId)
 	{
-		return WrappedEditor->GetSelectedTopLevelObjects();
+		WrappedEditor->RequestObjectColumnResort(ColumnId);
 	}
 
-	TArray<FSoftObjectPath> SDefaultReplicationStreamEditor::GetObjectsBeingPropertyEdited() const
+	void SDefaultReplicationStreamEditor::RequestPropertyColumnResort(const FName& ColumnId)
 	{
-		return WrappedEditor->GetObjectsBeingPropertyEdited();
+		WrappedEditor->RequestPropertyColumnResort(ColumnId);
 	}
 
-	bool SDefaultReplicationStreamEditor::SortPropertiesPredicate(
-		const FReplicatedPropertyData& Left,
-		const FReplicatedPropertyData& Right
-		) const
+	TArray<FSoftObjectPath> SDefaultReplicationStreamEditor::GetObjectsBeingPropertyEdited() const
 	{
-		return ReplicationColumns::Property::SortBySelectionThenByName_PropertyPredicate(
-			WrappedEditor->GetObjectsBeingPropertyEdited(),
-			*PropertiesModel,
-			Left,
-			Right
-			);
+		return WrappedEditor->GetObjectsBeingPropertyEdited();
 	}
 }
diff --git a/Engine/Plugins/Developer/Concert/ConcertUI/ConcertClientSharedSlate/Source/ConcertClientSharedSlate/Private/Replication/Editor/View/ObjectEditor/SDefaultReplicationStreamEditor.h b/Engine/Plugins/Developer/Concert/ConcertUI/ConcertClientSharedSlate/Source/ConcertClientSharedSlate/Private/Replication/Editor/View/ObjectEditor/SDefaultReplicationStreamEditor.h
--- a/Engine/Plugins/Developer/Concert/ConcertUI/ConcertClientSharedSlate/Source/ConcertClientSharedSlate/Private/Replication/Editor/View/ObjectEditor/SDefaultReplicationStreamEditor.h
+++ b/Engine/Plugins/Developer/Concert/ConcertUI/ConcertClientSharedSlate/Source/ConcertClientSharedSlate/Private/Replication/Editor/View/ObjectEditor/SDefaultReplicationStreamEditor.h
@@ -4,7 +4,9 @@
 
 #include "Replication/Editor/View/IReplicationStreamEditor.h"
 #include "Replication/Editor/View/PredefinedReplicationColumns.h"
+#include "Replication/Editor/View/ReplicationColumnsUtils.h"
 #include "Replication/ReplicationWidgetDelegates.h"
+#include "Replication/Editor/View/Tree/SelectionViewerColumns.h"
 
 #include "Widgets/DeclarativeSyntaxSupport.h"
 
@@ -35,8 +37,17 @@ namespace UE::ConcertClientSharedSlate
 		{}
 			/** Additional columns to add to the object view */
 			SLATE_ARGUMENT(TArray<ReplicationColumns::FReplicationTopLevelObjectColumn>, AdditionalObjectColumns)
+			/** Initial primary sort to set. */
+			SLATE_ARGUMENT(FColumnSortInfo, PrimaryObjectSort)
+			/** Initial secondary sort to set. */
+			SLATE_ARGUMENT(FColumnSortInfo, SecondaryObjectSort)
+			
 			/** Additional columns to add to the property view */
 			SLATE_ARGUMENT(TArray<ReplicationColumns::FReplicationPropertyColumn>, AdditionalPropertyColumns)
+			/** Initial primary sort to set. */
+			SLATE_ARGUMENT(FColumnSortInfo, PrimaryPropertySort)
+			/** Initial secondary sort to set. */
+			SLATE_ARGUMENT(FColumnSortInfo, SecondaryPropertySort)
 			
 			/** Optional. Placed between root object outliner and property editor. */
 			SLATE_ARGUMENT(TSharedPtr<IReplicationSubobjectView>, SubobjectView)
@@ -69,7 +80,8 @@ namespace UE::ConcertClientSharedSlate
 
 		//~ Begin IReplicationStreamEditor Interface
 		virtual void Refresh() override;
-		virtual TArray<FSoftObjectPath> GetSelectedTopLevelObjects() const override;
+		virtual void RequestObjectColumnResort(const FName& ColumnId) override;
+		virtual void RequestPropertyColumnResort(const FName& ColumnId) override;
 		virtual TArray<FSoftObjectPath> GetObjectsBeingPropertyEdited() const override;
 		//~ End IReplicationStreamEditor Interface
 
@@ -80,8 +92,6 @@ namespace UE::ConcertClientSharedSlate
 		
 		/** Model needed for sorting properties based on selection. */
 		TSharedPtr<IEditableReplicationStreamModel> PropertiesModel;
-		
-		bool SortPropertiesPredicate(const FReplicatedPropertyData& Left, const FReplicatedPropertyData& Right) const;
     };
 }
 
diff --git a/Engine/Plugins/Developer/Concert/ConcertUI/ConcertClientSharedSlate/Source/ConcertClientSharedSlate/Private/Replication/Editor/View/ObjectViewer/Property/SPropertyTreeView.cpp b/Engine/Plugins/Developer/Concert/ConcertUI/ConcertClientSharedSlate/Source/ConcertClientSharedSlate/Private/Replication/Editor/View/ObjectViewer/Property/SPropertyTreeView.cpp
--- a/Engine/Plugins/Developer/Concert/ConcertUI/ConcertClientSharedSlate/Source/ConcertClientSharedSlate/Private/Replication/Editor/View/ObjectViewer/Property/SPropertyTreeView.cpp
+++ b/Engine/Plugins/Developer/Concert/ConcertUI/ConcertClientSharedSlate/Source/ConcertClientSharedSlate/Private/Replication/Editor/View/ObjectViewer/Property/SPropertyTreeView.cpp
@@ -64,6 +64,8 @@ namespace UE::ConcertClientSharedSlate
 				.OnGetChildren(InArgs._OnGetChildren)
 				.Columns(InArgs._Columns)
 				.ExpandableColumnLabel(InArgs._ExpandableColumnLabel)
+				.PrimarySort(InArgs._PrimarySort)
+				.SecondarySort(InArgs._SecondarySort)
 				.SelectionMode(InArgs._SelectionMode)
 				.FilterItem(this, &SPropertyTreeView::PassesFilters)
 				.LeftOfSearchBar()
@@ -112,6 +114,11 @@ namespace UE::ConcertClientSharedSlate
 		ReplicatedProperties->OnItemsChanged();
 	}
 
+	void SPropertyTreeView::RequestResortForColumn(const FName& ColumnId)
+	{
+		ReplicatedProperties->RequestResortForColumn(ColumnId);
+	}
+
 	SPropertyTreeView::FBuildFilterBarResult SPropertyTreeView::BuildFilterBar()
 	{
 		TSharedRef<FFilterCategory> GeneralCategory = MakeShared<FFilterCategory>(
diff --git a/Engine/Plugins/Developer/Concert/ConcertUI/ConcertClientSharedSlate/Source/ConcertClientSharedSlate/Private/Replication/Editor/View/ObjectViewer/Property/SPropertyTreeView.h b/Engine/Plugins/Developer/Concert/ConcertUI/ConcertClientSharedSlate/Source/ConcertClientSharedSlate/Private/Replication/Editor/View/ObjectViewer/Property/SPropertyTreeView.h
--- a/Engine/Plugins/Developer/Concert/ConcertUI/ConcertClientSharedSlate/Source/ConcertClientSharedSlate/Private/Replication/Editor/View/ObjectViewer/Property/SPropertyTreeView.h
+++ b/Engine/Plugins/Developer/Concert/ConcertUI/ConcertClientSharedSlate/Source/ConcertClientSharedSlate/Private/Replication/Editor/View/ObjectViewer/Property/SPropertyTreeView.h
@@ -34,6 +34,10 @@ namespace UE::ConcertClientSharedSlate
 			SLATE_ARGUMENT(TArray<TReplicationColumn<FReplicatedPropertyData>>, Columns)
 			/** The name of the column that will have the SExpanderArrow for the tree view. */
 			SLATE_ARGUMENT(FName, ExpandableColumnLabel)
+			/** Initial primary sort to set. */
+			SLATE_ARGUMENT(FColumnSortInfo, PrimarySort)
+			/** Initial secondary sort to set. */
+			SLATE_ARGUMENT(FColumnSortInfo, SecondarySort)
 		
 			/** How many items are to allowed to be selected */
 			SLATE_ARGUMENT(ESelectionMode::Type, SelectionMode)
@@ -53,6 +57,9 @@ namespace UE::ConcertClientSharedSlate
 		/** Called when the items need to be refiltered due to the item source changing. */
 		void OnItemsChanged() const;
 
+		/** Requests that the given column be resorted, if it currently affects the row sorting. */
+		void RequestResortForColumn(const FName& ColumnId);
+
 	private:
 
 		/** The tree view displaying the replicated properties */
diff --git a/Engine/Plugins/Developer/Concert/ConcertUI/ConcertClientSharedSlate/Source/ConcertClientSharedSlate/Private/Replication/Editor/View/ObjectViewer/SReplicatedPropertyView.cpp b/Engine/Plugins/Developer/Concert/ConcertUI/ConcertClientSharedSlate/Source/ConcertClientSharedSlate/Private/Replication/Editor/View/ObjectViewer/SReplicatedPropertyView.cpp
--- a/Engine/Plugins/Developer/Concert/ConcertUI/ConcertClientSharedSlate/Source/ConcertClientSharedSlate/Private/Replication/Editor/View/ObjectViewer/SReplicatedPropertyView.cpp
+++ b/Engine/Plugins/Developer/Concert/ConcertUI/ConcertClientSharedSlate/Source/ConcertClientSharedSlate/Private/Replication/Editor/View/ObjectViewer/SReplicatedPropertyView.cpp
@@ -23,7 +23,6 @@ namespace UE::ConcertClientSharedSlate
 	void SReplicatedPropertyView::Construct(const FArguments& InArgs, TSharedRef<IReplicationStreamModel> InPropertiesModel)
 	{
 		PropertiesModel = MoveTemp(InPropertiesModel);
-		SortPropertyRowPredicate = InArgs._SortPropertyRowPredicate;
 		GetSelectedRootObjectsDelegate = InArgs._GetSelectedRootObjects;
 		check(GetSelectedRootObjectsDelegate.IsBound());
 		
@@ -83,6 +82,11 @@ namespace UE::ConcertClientSharedSlate
 		ReplicatedProperties->OnItemsChanged();
 	}
 
+	void SReplicatedPropertyView::RequestResortForColumn(const FName& ColumnId)
+	{
+		ReplicatedProperties->RequestResortForColumn(ColumnId);
+	}
+
 	TArray<FSoftObjectPath> SReplicatedPropertyView::GetObjectsSelectedForPropertyEditing() const
 	{
 		TArray<FSoftObjectPath> Result;
@@ -101,6 +105,14 @@ namespace UE::ConcertClientSharedSlate
 			ReplicationColumns::Property::TypeColumn()
 		};
 		Columns.Append(InArgs._AdditionalPropertyColumns);
+
+		// Set both primary and secondary in case one is overriden but always use the override.
+		const FColumnSortInfo PrimarySort = InArgs._PrimarySort.IsValid()
+			? InArgs._PrimarySort
+			: FColumnSortInfo{ ReplicationColumns::TopLevel::LabelColumnId, EColumnSortMode::Ascending };
+		const FColumnSortInfo SecondarySort = InArgs._SecondarySort.IsValid()
+			? InArgs._SecondarySort
+			: FColumnSortInfo{ ReplicationColumns::TopLevel::LabelColumnId, EColumnSortMode::Ascending };
 		
 		return SAssignNew(PropertyContent, SWidgetSwitcher)
 			// Make sure the slots are coherent with the order of EReplicatedPropertyContent!
@@ -114,6 +126,8 @@ namespace UE::ConcertClientSharedSlate
 				.OnGetChildren(this, &SReplicatedPropertyView::GetPropertyRowChildren)
 				.Columns(Columns)
 				.ExpandableColumnLabel(ReplicationColumns::Property::LabelColumnId)
+				.PrimarySort(PrimarySort)
+				.SecondarySort(SecondarySort)
 				.SelectionMode(ESelectionMode::Multi)
 				.LeftOfSearchBar()
 				[
@@ -156,8 +170,6 @@ namespace UE::ConcertClientSharedSlate
 				RootPropertyRowData.Emplace(PropertyData);
 			}
 		}
-
-		SortPropertyRowArray(RootPropertyRowData);
 	}
 	
 	TOptional<FSoftClassPath> SReplicatedPropertyView::GetClassForPropertiesFromSelection(const TArray<FSoftObjectPath>& Objects) const
@@ -183,21 +195,9 @@ namespace UE::ConcertClientSharedSlate
 			}
 		}
 
-		SortPropertyRowArray(Children);
 		Algo::ForEach(Children, [&ProcessChild](const TSharedPtr<FReplicatedPropertyData>& Data){ ProcessChild(Data); });
 	}
 
-	void SReplicatedPropertyView::SortPropertyRowArray(TArray<TSharedPtr<FReplicatedPropertyData>>& ToSort) const
-	{
-		if (SortPropertyRowPredicate.IsBound())
-		{
-			ToSort.Sort([this](const TSharedPtr<FReplicatedPropertyData>& Left, const TSharedPtr<FReplicatedPropertyData>& Right)
-			{
-				return SortPropertyRowPredicate.Execute(*Left.Get(), *Right.Get());
-			});
-		}
-	}
-
 	void SReplicatedPropertyView::SetPropertyContent(EReplicatedPropertyContent Content) const
 	{
 		PropertyContent->SetActiveWidgetIndex(static_cast<int32>(Content));
diff --git a/Engine/Plugins/Developer/Concert/ConcertUI/ConcertClientSharedSlate/Source/ConcertClientSharedSlate/Private/Replication/Editor/View/ObjectViewer/SReplicatedPropertyView.h b/Engine/Plugins/Developer/Concert/ConcertUI/ConcertClientSharedSlate/Source/ConcertClientSharedSlate/Private/Replication/Editor/View/ObjectViewer/SReplicatedPropertyView.h
--- a/Engine/Plugins/Developer/Concert/ConcertUI/ConcertClientSharedSlate/Source/ConcertClientSharedSlate/Private/Replication/Editor/View/ObjectViewer/SReplicatedPropertyView.h
+++ b/Engine/Plugins/Developer/Concert/ConcertUI/ConcertClientSharedSlate/Source/ConcertClientSharedSlate/Private/Replication/Editor/View/ObjectViewer/SReplicatedPropertyView.h
@@ -3,10 +3,9 @@
 #pragma once
 
 #include "Replication/Editor/View/Tree/SelectionViewerColumns.h"
-#include "Replication/ReplicationWidgetDelegates.h"
+#include "Replication/Data/ConcertPropertySelection.h"
 
 #include "Misc/Optional.h"
-#include "Replication/Data/ConcertPropertySelection.h"
 #include "Widgets/DeclarativeSyntaxSupport.h"
 #include "Widgets/SCompoundWidget.h"
 
@@ -27,9 +26,11 @@ namespace UE::ConcertClientSharedSlate
 		{}
 			/** Additional columns to add to the property view */
 			SLATE_ARGUMENT(TArray<ReplicationColumns::FReplicationPropertyColumn>, AdditionalPropertyColumns)
+			/** Initial primary sort to set. */
+			SLATE_ARGUMENT(FColumnSortInfo, PrimarySort)
+			/** Initial secondary sort to set. Defaults to Label column. */
+			SLATE_ARGUMENT(FColumnSortInfo, SecondarySort)
 		
-			/** Optional. Used for determining the order in which properties are displayed. */
-			SLATE_EVENT(FSortPropertyPredicate, SortPropertyRowPredicate)
 			/** Gets the root objects selected in the object outliner. */
 			SLATE_EVENT(FGetSelectedRootObjects, GetSelectedRootObjects)
 		
@@ -40,6 +41,8 @@ namespace UE::ConcertClientSharedSlate
 		void Construct(const FArguments& InArgs, TSharedRef<IReplicationStreamModel> InPropertiesModel);
 		
 		void RefreshPropertyData();
+		/** Requests that the given column be resorted, if it currently affects the row sorting. */
+		void RequestResortForColumn(const FName& ColumnId);
 		
 		const TArray<TSharedPtr<FReplicatedPropertyData>>& GetPropertyRowData() const { return PropertyRowData; }
 		TArray<FSoftObjectPath> GetObjectsSelectedForPropertyEditing() const;
@@ -64,8 +67,6 @@ namespace UE::ConcertClientSharedSlate
 		/** Determines the content displayed for PropertyArea. */
 		TSharedPtr<SWidgetSwitcher> PropertyContent;
 		
-		/** Used for determining the order in which properties are displayed. Default: Sort by name. */
-		FSortPropertyPredicate SortPropertyRowPredicate;
 		/** Gets the root objects selected in the object outliner. */
 		FGetSelectedRootObjects GetSelectedRootObjectsDelegate;
 		
@@ -92,7 +93,6 @@ namespace UE::ConcertClientSharedSlate
 		void GetPropertyRowChildren(TSharedPtr<FReplicatedPropertyData> ReplicatedPropertyData, TFunctionRef<void(TSharedPtr<FReplicatedPropertyData>)> ProcessChild);
 		
 		// Utils
-		void SortPropertyRowArray(TArray<TSharedPtr<FReplicatedPropertyData>>& ToSort) const;
 		void SetPropertyContent(EReplicatedPropertyContent Content) const;
 	};
 }
diff --git a/Engine/Plugins/Developer/Concert/ConcertUI/ConcertClientSharedSlate/Source/ConcertClientSharedSlate/Private/Replication/Editor/View/ObjectViewer/SReplicationStreamViewer.cpp b/Engine/Plugins/Developer/Concert/ConcertUI/ConcertClientSharedSlate/Source/ConcertClientSharedSlate/Private/Replication/Editor/View/ObjectViewer/SReplicationStreamViewer.cpp
--- a/Engine/Plugins/Developer/Concert/ConcertUI/ConcertClientSharedSlate/Source/ConcertClientSharedSlate/Private/Replication/Editor/View/ObjectViewer/SReplicationStreamViewer.cpp
+++ b/Engine/Plugins/Developer/Concert/ConcertUI/ConcertClientSharedSlate/Source/ConcertClientSharedSlate/Private/Replication/Editor/View/ObjectViewer/SReplicationStreamViewer.cpp
@@ -40,19 +40,19 @@ namespace UE::ConcertClientSharedSlate
 		RefreshPropertyData();
 	}
 
-	TArray<FSoftObjectPath> SReplicationStreamViewer::GetSelectedTopLevelObjects() const
+	void SReplicationStreamViewer::RequestObjectColumnResort(const FName& ColumnId)
 	{
-		TArray<FSoftObjectPath> Result;
-		Algo::Transform(GetSelectedOutlinerObjects(), Result, [](const TSharedPtr<FReplicatedObjectData>& Data)
-		{
-			return Data->GetObjectPath();
-		});
-		return Result;
+		ReplicatedObjects->RequestResortForColumn(ColumnId);
+	}
+
+	void SReplicationStreamViewer::RequestPropertyColumnResort(const FName& ColumnId)
+	{
+		PropertySection->RequestResortForColumn(ColumnId);
 	}
 
 	TArray<FSoftObjectPath> SReplicationStreamViewer::GetObjectsBeingPropertyEdited() const
 	{
-		return SubobjectAndPropertySection->GetObjectsSelectedForPropertyEditing();
+		return PropertySection->GetObjectsSelectedForPropertyEditing();
 	}
 
 	void SReplicationStreamViewer::RefreshObjectData()
@@ -93,7 +93,7 @@ namespace UE::ConcertClientSharedSlate
 
 	void SReplicationStreamViewer::RefreshPropertyData()
 	{
-		SubobjectAndPropertySection->RefreshPropertyData();
+		PropertySection->RefreshPropertyData();
 	}
 
 	void SReplicationStreamViewer::SelectTopLevelObjects(TConstArrayView<FSoftObjectPath> Objects)
@@ -182,6 +182,14 @@ namespace UE::ConcertClientSharedSlate
 		const bool bHasNoOutlinerObjectsAttribute = InArgs._NoOutlinerObjects.IsBound() || InArgs._NoOutlinerObjects.IsSet(); 
 		const TAttribute<FText> NoObjectsAttribute = bHasNoOutlinerObjectsAttribute ? InArgs._NoOutlinerObjects : LOCTEXT("NoObjects", "No objects to display");
 
+		// Set both primary and secondary in case one is overriden but always use the override.
+		const FColumnSortInfo PrimaryObjectSort = InArgs._PrimaryObjectSort.IsValid()
+			? InArgs._PrimaryObjectSort
+			: FColumnSortInfo{ ReplicationColumns::TopLevel::LabelColumnId, EColumnSortMode::Ascending };
+		const FColumnSortInfo SecondaryObjectSort = InArgs._SecondaryObjectSort.IsValid()
+			? InArgs._SecondaryObjectSort
+			: FColumnSortInfo{ ReplicationColumns::TopLevel::LabelColumnId, EColumnSortMode::Ascending };
+		
 		return SAssignNew(ReplicatedObjects, SReplicationTreeView<FReplicatedObjectData>)
 			.RootItemsSource(&RootObjectRowData)
 			.OnGetChildren(this, &SReplicationStreamViewer::GetObjectRowChildren)
@@ -192,6 +200,8 @@ namespace UE::ConcertClientSharedSlate
 			})
 			.Columns(Columns)
 			.ExpandableColumnLabel(ReplicationColumns::TopLevel::LabelColumnId)
+			.PrimarySort(PrimaryObjectSort)
+			.SecondarySort(SecondaryObjectSort)
 			.SelectionMode(ESelectionMode::Multi)
 			.LeftOfSearchBar() [ InArgs._LeftOfObjectSearchBar.Widget ]
 			.NoItemsContent() [ SNew(STextBlock).Text(NoObjectsAttribute) ];
@@ -220,9 +230,10 @@ namespace UE::ConcertClientSharedSlate
 				]
 				.BodyContent()
 				[
-					SAssignNew(SubobjectAndPropertySection, SReplicatedPropertyView, PropertiesModel.ToSharedRef())
+					SAssignNew(PropertySection, SReplicatedPropertyView, PropertiesModel.ToSharedRef())
 					.AdditionalPropertyColumns(InArgs._AdditionalPropertyColumns)
-					.SortPropertyRowPredicate(InArgs._SortPropertyRowPredicate)
+					.PrimarySort(InArgs._PrimaryPropertySort)
+					.SecondarySort(InArgs._SecondaryPropertySort)
 					.GetSelectedRootObjects_Lambda([this](){ return GetSelectedOutlinerObjects(); })
 					.LeftOfPropertySearchBar()
 					[
diff --git a/Engine/Plugins/Developer/Concert/ConcertUI/ConcertClientSharedSlate/Source/ConcertClientSharedSlate/Private/Replication/Editor/View/ObjectViewer/SReplicationStreamViewer.h b/Engine/Plugins/Developer/Concert/ConcertUI/ConcertClientSharedSlate/Source/ConcertClientSharedSlate/Private/Replication/Editor/View/ObjectViewer/SReplicationStreamViewer.h
--- a/Engine/Plugins/Developer/Concert/ConcertUI/ConcertClientSharedSlate/Source/ConcertClientSharedSlate/Private/Replication/Editor/View/ObjectViewer/SReplicationStreamViewer.h
+++ b/Engine/Plugins/Developer/Concert/ConcertUI/ConcertClientSharedSlate/Source/ConcertClientSharedSlate/Private/Replication/Editor/View/ObjectViewer/SReplicationStreamViewer.h
@@ -47,8 +47,17 @@ namespace UE::ConcertClientSharedSlate
 		{}
 			/** Additional columns to add to the object view */
 			SLATE_ARGUMENT(TArray<ReplicationColumns::FReplicationTopLevelObjectColumn>, AdditionalObjectColumns)
+			/** Initial primary sort to set. */
+			SLATE_ARGUMENT(FColumnSortInfo, PrimaryObjectSort)
+			/** Initial secondary sort to set. */
+			SLATE_ARGUMENT(FColumnSortInfo, SecondaryObjectSort)
+		
 			/** Additional columns to add to the property view */
 			SLATE_ARGUMENT(TArray<ReplicationColumns::FReplicationPropertyColumn>, AdditionalPropertyColumns)
+			/** Initial primary sort to set. */
+			SLATE_ARGUMENT(FColumnSortInfo, PrimaryPropertySort)
+			/** Initial secondary sort to set. */
+			SLATE_ARGUMENT(FColumnSortInfo, SecondaryPropertySort)
 
 			/** Optional. If set, this determines the children nested under the root objects. */
 			SLATE_ARGUMENT(TSharedPtr<ISubobjectModel>, SubobjectModel)
@@ -58,9 +67,6 @@ namespace UE::ConcertClientSharedSlate
 		
 			/** Called to generate the context menu for objects. */
 			SLATE_EVENT(FOnContextMenuOpening, OnObjectsContextMenuOpening)
-
-			/** Optional. Used for determining the order in which properties are displayed. */
-			SLATE_EVENT(FSortPropertyPredicate, SortPropertyRowPredicate)
 		
 			/** Optional widget to add to the left of the object list search bar. */
 			SLATE_NAMED_SLOT(FArguments, LeftOfObjectSearchBar)
@@ -75,7 +81,8 @@ namespace UE::ConcertClientSharedSlate
 
 		//~ Begin IReplicationStreamViewer Interface
 		virtual void Refresh() override;
-		virtual TArray<FSoftObjectPath> GetSelectedTopLevelObjects() const override;
+		virtual void RequestObjectColumnResort(const FName& ColumnId) override;
+		virtual void RequestPropertyColumnResort(const FName& ColumnId) override;
 		virtual TArray<FSoftObjectPath> GetObjectsBeingPropertyEdited() const override;
 		//~ End IReplicationStreamViewer Interface
 
@@ -101,7 +108,7 @@ namespace UE::ConcertClientSharedSlate
 		/** Lists the properties of the selected actor */
 		TSharedPtr<SExpandableArea> PropertyArea;
 		/** Edits the property list and (optionally) exposes subobjects of the selected root object. */
-		TSharedPtr<SReplicatedPropertyView> SubobjectAndPropertySection;
+		TSharedPtr<SReplicatedPropertyView> PropertySection;
 
 		/** Tree view for replicated objects. */
 		TSharedPtr<SReplicationTreeView<FReplicatedObjectData>> ReplicatedObjects;
diff --git a/Engine/Plugins/Developer/Concert/ConcertUI/ConcertClientSharedSlate/Source/ConcertClientSharedSlate/Private/Replication/Editor/View/Tree/SReplicationTreeView.h b/Engine/Plugins/Developer/Concert/ConcertUI/ConcertClientSharedSlate/Source/ConcertClientSharedSlate/Private/Replication/Editor/View/Tree/SReplicationTreeView.h
--- a/Engine/Plugins/Developer/Concert/ConcertUI/ConcertClientSharedSlate/Source/ConcertClientSharedSlate/Private/Replication/Editor/View/Tree/SReplicationTreeView.h
+++ b/Engine/Plugins/Developer/Concert/ConcertUI/ConcertClientSharedSlate/Source/ConcertClientSharedSlate/Private/Replication/Editor/View/Tree/SReplicationTreeView.h
@@ -3,6 +3,7 @@
 #pragma once
 
 #include "Replication/Editor/View/ReplicationColumn.h"
+#include "Replication/Editor/View/ReplicationColumnsUtils.h"
 #include "SReplicationColumnRow.h"
 
 #include "Algo/RemoveIf.h"
@@ -78,6 +79,10 @@ namespace UE::ConcertClientSharedSlate
 			SLATE_ARGUMENT(FName, ExpandableColumnLabel)
 			/** Visibility of the header row */
 			SLATE_ARGUMENT(EVisibility, HeaderRowVisibility)
+			/** Initial primary sort to set. */
+			SLATE_ARGUMENT(FColumnSortInfo, PrimarySort)
+			/** Initial secondary sort to set. */
+			SLATE_ARGUMENT(FColumnSortInfo, SecondarySort)
 		
 			/** How many items are to allowed to be selected */
 			SLATE_ARGUMENT(ESelectionMode::Type, SelectionMode)
@@ -105,7 +110,7 @@ namespace UE::ConcertClientSharedSlate
 			
 			SearchText = MakeShared<FText>();
 			SearchTextFilter = MakeShared<TTextFilter<const TSharedPtr<TItemType>&>>(TTextFilter<const TSharedPtr<TItemType>&>::FItemToStringArray::CreateSP(this, &SReplicationTreeView::PopulateSearchStrings));
-			SearchTextFilter->OnChanged().AddSP(this, &SReplicationTreeView::OnFilterChanged);
+			SearchTextFilter->OnChanged().AddSP(this, &SReplicationTreeView::RequestRefilter);
 			
 			ChildSlot
 			[
@@ -158,12 +163,82 @@ namespace UE::ConcertClientSharedSlate
 					]
 				]
 			];
+
+			if (!InArgs._PrimarySort.SortedColumnId.IsNone())
+			{
+				SetPrimarySortMode(
+					InArgs._PrimarySort.SortedColumnId,
+					// In case None was specified by accident, use good defaults or no sorting will occur
+					InArgs._PrimarySort.SortMode == EColumnSortMode::None ? EColumnSortMode::Ascending : InArgs._PrimarySort.SortMode
+					);
+			}
+			if (!InArgs._SecondarySort.SortedColumnId.IsNone())
+			{
+				SetSecondarySortMode(
+					InArgs._SecondarySort.SortedColumnId, 
+					// In case None was specified by accident, use good defaults or no sorting will occur
+					InArgs._SecondarySort.SortMode == EColumnSortMode::None ? EColumnSortMode::Ascending : InArgs._SecondarySort.SortMode
+					);
+			}
 		}
 		
 		void OnItemsChanged()
 		{
-			// Re-filter everything. There should not be many items so filtering everything again should be fine. Calls RequestListRefresh as well.
-			OnFilterChanged();
+			// Re-filter everything. There should not be many items so filtering everything again should be fine
+			bFilterChanged = true;
+			bRequestedSort = true;
+		}
+		void RequestRefilter()
+		{
+			bFilterChanged = true;
+			bRequestedSort = true;
+		}
+		void RequestResort()
+		{
+			bRequestedSort = true;
+		}
+		
+		/** Requests that the given column be resorted, if it currently affects the row sorting. */
+		void RequestResortForColumn(const FName& ColumnId)
+		{
+			if (PrimarySortInfo.SortedColumnId == ColumnId || SecondarySortInfo.SortedColumnId == ColumnId)
+			{
+				RequestResort();
+			}
+		}
+
+		void SetPrimarySortMode(FName SortedColumnId, EColumnSortMode::Type SortMode)
+		{
+			const TReplicationColumn<TItemType>* Column = FindColumnByName(SortedColumnId);
+			if (ensure(Column && Column->CanBeSorted()))
+			{
+				PrimarySortInfo = { SortedColumnId, SortMode };
+				
+				if (SortedColumnId == SecondarySortInfo.SortedColumnId) // Cannot be primary and secondary at the same time.
+				{
+					SecondarySortInfo.SortedColumnId = FName();
+					SecondarySortInfo.SortMode = EColumnSortMode::None;
+				}
+			}
+		}
+		void SetSecondarySortMode(FName SortedColumnId, EColumnSortMode::Type SortMode)
+		{
+			if (!PrimarySortInfo.IsValid())
+			{
+				SetPrimarySortMode(SortedColumnId, SortMode);
+				return;
+			}
+
+			if (PrimarySortInfo.SortedColumnId == SortedColumnId)
+			{
+				return;
+			}
+			
+			const TReplicationColumn<TItemType>* Column = FindColumnByName(SortedColumnId);
+			if (ensure(Column && Column->CanBeSorted()))
+			{
+				SecondarySortInfo = { SortedColumnId, SortMode };
+			}
 		}
 		
 		void SetSelectedItems(const TArray<TSharedPtr<TItemType>>& ObjectsToSelect, bool bIsSelected)
@@ -183,6 +258,7 @@ namespace UE::ConcertClientSharedSlate
 		const TArray<TSharedPtr<TItemType>>& GetFilteredRootItems() const { return FilteredRootItems; }
 		
 		virtual FReply OnKeyDown(const FGeometry& MyGeometry, const FKeyEvent& InKeyEvent) override;
+		virtual void Tick(const FGeometry& AllottedGeometry, const double InCurrentTime, const float InDeltaTime) override;
 
 	private:
 
@@ -203,6 +279,12 @@ namespace UE::ConcertClientSharedSlate
 		TArray<TSharedPtr<TItemType>>* AllRootItems = nullptr;
 		TArray<TSharedPtr<TItemType>> FilteredRootItems;
 
+		bool bFilterChanged = false;
+		bool bRequestedSort = false;
+
+		FColumnSortInfo PrimarySortInfo;
+		FColumnSortInfo SecondarySortInfo;
+
 		/** Callback for getting an item's children. */
 		FGetItemChildren OnGetChildrenDelegate;
 		/** Optional delegate for responding to pressing the delete button */
@@ -213,19 +295,31 @@ namespace UE::ConcertClientSharedSlate
 		TOverrideColumnWidget OverrideColumnWidget;
 		/** Optional callback for determining whether this item can be filtered. If false, it will not be shown when searched. */
 		FIsSearchableItem IsSearchableItemDelegate;
-		
 
+		// STreeView creation
 		TSharedRef<SWidget> CreateTreeView(const FArguments& InArgs);
 		TSharedRef<SHeaderRow> CreateHeaderRow(const FArguments& InArgs);
 		TSharedRef<ITableRow> OnGenerateRowWidget(TSharedPtr<TItemType> Item, const TSharedRef<STableViewBase>& OwnerTable);
 		void GetRowChildren(TSharedPtr<TItemType> Item, TArray<TSharedPtr<TItemType>>& OutChildren);
 		
+		// Sorting callbacks
+		EColumnSortPriority::Type GetColumnSortPriority(const FName ColumnId) const;
+		EColumnSortMode::Type GetColumnSortMode(const FName ColumnId) const;
+		void OnColumnSortModeChanged(const EColumnSortPriority::Type SortPriority, const FName& ColumnId, const EColumnSortMode::Type InSortMode);
+
+		// Searching callbacks
 		void OnSearchTextCommitted(const FText& InFilterText, ETextCommit::Type CommitType);
 		void OnSearchTextChanged(const FText& InSearchText);
 
+		// Filtering
 		void PopulateSearchStrings(const TSharedPtr<TItemType>& Item, TArray<FString>& OutSearchStrings);
-		void OnFilterChanged();
+		void ReapplyFilters();
 		bool PassesFilters(const TSharedPtr<TItemType>& Item);
+
+		// Sorting
+		void Resort();
+		void Sort(TArray<TSharedPtr<TItemType>>& Items);
+		const TReplicationColumn<TItemType>* FindColumnByName(const FName& ColumnId) const;
 	};
 
 	template <typename TItemType>
@@ -240,6 +334,22 @@ namespace UE::ConcertClientSharedSlate
 		return SCompoundWidget::OnKeyDown(MyGeometry, InKeyEvent);
 	}
 
+	template <typename TItemType>
+	void SReplicationTreeView<TItemType>::Tick(const FGeometry& AllottedGeometry, const double InCurrentTime, const float InDeltaTime)
+	{
+		if (bFilterChanged)
+		{
+			ReapplyFilters();
+		}
+
+		if (bRequestedSort)
+		{
+			Resort();
+		}
+		
+		SCompoundWidget::Tick(AllottedGeometry, InCurrentTime, InDeltaTime);
+	}
+
 	template <typename TItemType>
 	TSharedRef<SWidget> SReplicationTreeView<TItemType>::CreateTreeView(const FArguments& InArgs)
 	{
@@ -296,11 +406,15 @@ namespace UE::ConcertClientSharedSlate
 		TSet<FName> DuplicateColumnDetection;
 		for (TReplicationColumn<TItemType>& Column : Columns)
 		{
-			check(!DuplicateColumnDetection.Contains(Column.ColumnId));
-			DuplicateColumnDetection.Add(Column.ColumnId);
+			const FName ColumnId = Column.ColumnId;
+			check(!DuplicateColumnDetection.Contains(ColumnId));
+			DuplicateColumnDetection.Add(ColumnId);
 			
 			// SHeaderRow owns the columns and deletes them when destroyed
 			TReplicationColumn<TItemType>* ManagedByHeaderRow = new TReplicationColumn<TItemType>(MoveTemp(Column));
+			ManagedByHeaderRow->SortPriority.Bind(TAttribute<EColumnSortPriority::Type>::FGetter::CreateSP(this, &SReplicationTreeView::GetColumnSortPriority, ColumnId));
+			ManagedByHeaderRow->SortMode.Bind(TAttribute<EColumnSortMode::Type>::FGetter::CreateSP(this, &SReplicationTreeView::GetColumnSortMode, ColumnId));
+			ManagedByHeaderRow->OnSortModeChanged.BindSP(this, &SReplicationTreeView::OnColumnSortModeChanged);
 			HeaderRow->AddColumn(*ManagedByHeaderRow);
 		}
 
@@ -310,19 +424,11 @@ namespace UE::ConcertClientSharedSlate
 	template <typename TItemType>
 	TSharedRef<ITableRow> SReplicationTreeView<TItemType>::OnGenerateRowWidget(TSharedPtr<TItemType> Item, const TSharedRef<STableViewBase>& OwnerTable)
 	{
-		const typename SReplicationColumnRow<TItemType>::FGetColumn ColumnGetter = SReplicationColumnRow<TItemType>::FGetColumn::CreateLambda([this](const FName& ColumnId) -> const TReplicationColumn<TItemType>*
-		{
-			for (const SHeaderRow::FColumn& Column : HeaderRow->GetColumns())
-			{
-				if (Column.ColumnId == ColumnId)
-				{
-					return static_cast<const TReplicationColumn<TItemType>*>(&Column);
-				}
-			}
-
-			checkNoEntry();
-			return nullptr;
-		});
+		const typename SReplicationColumnRow<TItemType>::FGetColumn ColumnGetter =
+			SReplicationColumnRow<TItemType>::FGetColumn::CreateSP(
+				this,
+				&SReplicationTreeView::FindColumnByName
+			);
 	
 		return SNew(SReplicationColumnRow<TItemType>, OwnerTable)
 			.HighlightText(SearchText)
@@ -344,7 +450,65 @@ namespace UE::ConcertClientSharedSlate
 					OutChildren.Add(ItemToAdd);
 				}
 			});
+			
+			Sort(OutChildren);
+		}
+	}
+
+	template <typename TItemType>
+	EColumnSortPriority::Type SReplicationTreeView<TItemType>::GetColumnSortPriority(const FName ColumnId) const
+	{
+		if (ColumnId == PrimarySortInfo.SortedColumnId)
+		{
+			return EColumnSortPriority::Primary;
+		}
+		if (ColumnId == SecondarySortInfo.SortedColumnId)
+		{
+			return EColumnSortPriority::Secondary;
+		}
+
+		return EColumnSortPriority::Max; // No specific priority.
+	}
+
+	template <typename TItemType>
+	EColumnSortMode::Type SReplicationTreeView<TItemType>::GetColumnSortMode(const FName ColumnId) const
+	{
+		if (ColumnId == PrimarySortInfo.SortedColumnId && PrimarySortInfo.IsValid())
+		{
+			return PrimarySortInfo.SortMode;
+		}
+		if (ColumnId == SecondarySortInfo.SortedColumnId && SecondarySortInfo.IsValid())
+		{
+			return SecondarySortInfo.SortMode;
+		}
+		return EColumnSortMode::None;
+	}
+
+	template <typename TItemType>
+	void SReplicationTreeView<TItemType>::OnColumnSortModeChanged(
+		const EColumnSortPriority::Type SortPriority,
+		const FName& ColumnId,
+		const EColumnSortMode::Type InSortMode
+		)
+	{
+		const TReplicationColumn<TItemType>* Column = FindColumnByName(ColumnId);
+		if (!ensure(Column)
+			// Cannot bind
+			|| !Column->CanBeSorted())
+		{
+			return;
+		}
+		
+		if (SortPriority == EColumnSortPriority::Primary)
+		{
+			SetPrimarySortMode(ColumnId, InSortMode);
 		}
+		else if (SortPriority == EColumnSortPriority::Secondary)
+		{
+			SetSecondarySortMode(ColumnId, InSortMode);
+		}
+
+		RequestResort();
 	}
 
 	template <typename TItemType>
@@ -380,7 +544,7 @@ namespace UE::ConcertClientSharedSlate
 	}
 
 	template <typename TItemType>
-	void SReplicationTreeView<TItemType>::OnFilterChanged()
+	void SReplicationTreeView<TItemType>::ReapplyFilters()
 	{
 		// Try preserving the selected activity.
 		TArray<TSharedPtr<TItemType>> SelectedItems = TreeView->GetSelectedItems();
@@ -405,6 +569,8 @@ namespace UE::ConcertClientSharedSlate
 			TreeView->RequestScrollIntoView(SelectedItems[0]);
 		}
 
+		bFilterChanged = false;
+		RequestResort();
 		TreeView->RequestListRefresh();
 	}
 
@@ -414,6 +580,68 @@ namespace UE::ConcertClientSharedSlate
 		return SearchTextFilter->PassesFilter(Item)
 			&& (!CustomFilterDelegate.IsBound() || CustomFilterDelegate.Execute(Item));
 	}
+
+	template <typename TItemType>
+	void SReplicationTreeView<TItemType>::Resort()
+	{
+		Sort(FilteredRootItems);
+		// GetRowChildren will be called again, which will do the resort the children.
+		TreeView->RequestListRefresh();
+	}
+
+	template <typename TItemType>
+	void SReplicationTreeView<TItemType>::Sort(TArray<TSharedPtr<TItemType>>& Items)
+	{
+		const auto IsLessThan = [this](const TSharedPtr<TItemType>& Left, const TSharedPtr<TItemType>& Right, const FName& ColumnName, EColumnSortMode::Type SortMode)
+		{
+			const TReplicationColumn<TItemType>* Column = FindColumnByName(ColumnName);
+			if (!ensure(Column) || !ensureMsgf(Column->CanBeSorted(), TEXT("Validate why invariant was broken.")))
+			{
+				return false;
+			}
+
+			switch (SortMode)
+			{
+				case EColumnSortMode::Ascending: return Column->IsLessThan(*Left, *Right);
+				case EColumnSortMode::Descending: return Column->IsLessThan(*Right, *Left);
+				case EColumnSortMode::None:
+				default: return false;
+			};
+		};
+		
+		Items.Sort([this, &IsLessThan](const TSharedPtr<TItemType>& Left, const TSharedPtr<TItemType>& Right)
+		{
+			if (PrimarySortInfo.IsValid() && IsLessThan(Left, Right, PrimarySortInfo.SortedColumnId, PrimarySortInfo.SortMode))
+			{
+				return true; // Left comes before Right
+			}
+			// Check for equality by inverting
+			if (PrimarySortInfo.IsValid() && IsLessThan(Right, Left, PrimarySortInfo.SortedColumnId, PrimarySortInfo.SortMode))
+			{
+				return false; // Right comes before Left
+			}
+			
+			// Lhs == Rhs on the primary column, need to order according to the secondary column if one is set.
+			return !SecondarySortInfo.IsValid()
+				? false
+				: IsLessThan(Left, Right, SecondarySortInfo.SortedColumnId, SecondarySortInfo.SortMode);
+		});
+
+		bRequestedSort = false;
+	}
+
+	template <typename TItemType>
+	const TReplicationColumn<TItemType>* SReplicationTreeView<TItemType>::FindColumnByName(const FName& ColumnId) const
+	{
+		for (const SHeaderRow::FColumn& Column : HeaderRow->GetColumns())
+		{
+			if (Column.ColumnId == ColumnId)
+			{
+				return static_cast<const TReplicationColumn<TItemType>*>(&Column);
+			}
+		}
+		return nullptr;
+	}
 }
 
 #undef LOCTEXT_NAMESPACE
\ No newline at end of file
diff --git a/Engine/Plugins/Developer/Concert/ConcertUI/ConcertClientSharedSlate/Source/ConcertClientSharedSlate/Private/Replication/Editor/View/Tree/SelectionViewerColumns.cpp b/Engine/Plugins/Developer/Concert/ConcertUI/ConcertClientSharedSlate/Source/ConcertClientSharedSlate/Private/Replication/Editor/View/Tree/SelectionViewerColumns.cpp
--- a/Engine/Plugins/Developer/Concert/ConcertUI/ConcertClientSharedSlate/Source/ConcertClientSharedSlate/Private/Replication/Editor/View/Tree/SelectionViewerColumns.cpp
+++ b/Engine/Plugins/Developer/Concert/ConcertUI/ConcertClientSharedSlate/Source/ConcertClientSharedSlate/Private/Replication/Editor/View/Tree/SelectionViewerColumns.cpp
@@ -31,20 +31,26 @@ namespace UE::ConcertClientSharedSlate::ReplicationColumns::TopLevel
 	
 	FReplicationTopLevelObjectColumn LabelColumn(TSharedRef<IReplicationStreamModel> Model, ISubobjectModel* SubobjectModel)
 	{
+		auto GetDisplayText = [SubobjectModel](const FReplicatedObjectData& ObjectData)
+		{
+			const FSoftObjectPath& ObjectPath = ObjectData.GetObjectPath();
+			FText Text = DisplayUtils::GetObjectDisplayText(ObjectPath);
+			if (SubobjectModel && !SubobjectModel->IsTopLevelObject(ObjectPath))
+			{
+				if (const TOptional<FSoftObjectPath> OwningActorPath = ObjectUtils::GetActorOf(ObjectPath))
+				{
+					SubobjectModel->SetTopLevelObject(*OwningActorPath);
+					Text = SubobjectModel->GetSubobjectDisplayName(ObjectPath);
+				}
+			}
+			return Text;
+		};
+		
 		return FReplicationTopLevelObjectColumn(
 			FReplicationTopLevelObjectColumn::FArguments()
-				.GenerateWidgetColumn_Lambda([Model = MoveTemp(Model), SubobjectModel](const FReplicationTopLevelObjectColumn::FBuildArgs& Args)
+				.GenerateWidgetColumn_Lambda([Model = MoveTemp(Model), GetDisplayText](const FReplicationTopLevelObjectColumn::FBuildArgs& Args)
 				{
-					const FSoftObjectPath& ObjectPath = Args.RowData.GetObjectPath();
-					FText Text = DisplayUtils::GetObjectDisplayText(ObjectPath);
-					if (SubobjectModel && !SubobjectModel->IsTopLevelObject(ObjectPath))
-					{
-						if (const TOptional<FSoftObjectPath> OwningActorPath = ObjectUtils::GetActorOf(ObjectPath))
-						{
-							SubobjectModel->SetTopLevelObject(*OwningActorPath);
-							Text = SubobjectModel->GetSubobjectDisplayName(ObjectPath);
-						}
-					}
+					const FText Text = GetDisplayText(Args.RowData);
 					
 					return SNew(SHorizontalBox)
 						+SHorizontalBox::Slot()
@@ -70,6 +76,10 @@ namespace UE::ConcertClientSharedSlate::ReplicationColumns::TopLevel
 				{
 					InOutSearchStrings.Add(DisplayUtils::GetObjectDisplayText(ObjectData.GetObjectPath()).ToString());
 				})
+				.IsLessThan_Lambda([GetDisplayText](const FReplicatedObjectData& Left, const FReplicatedObjectData& Right)
+				{
+					return GetDisplayText(Left).ToString() < GetDisplayText(Right).ToString();
+				})
 				.ColumnSortOrder(static_cast<int32>(ETopLevelColumnOrder::Label)),
 			SHeaderRow::Column(LabelColumnId)
 				.DefaultLabel(LOCTEXT("LabelColumnLabel", "Label"))
@@ -91,6 +101,10 @@ namespace UE::ConcertClientSharedSlate::ReplicationColumns::TopLevel
 				{
 					InOutSearchStrings.Add(DisplayUtils::GetObjectTypeText(Model.Get(), ObjectData.GetObjectPath()).ToString());
 				})
+				.IsLessThan_Lambda([Model](const FReplicatedObjectData& Left, const FReplicatedObjectData& Right)
+				{
+					return DisplayUtils::GetObjectTypeText(*Model, Left.GetObjectPath()).ToString() < DisplayUtils::GetObjectTypeText(*Model, Right.GetObjectPath()).ToString();
+				})
 				.ColumnSortOrder(static_cast<int32>(ETopLevelColumnOrder::Type)),
 			SHeaderRow::Column(TypeColumnId)
 				.DefaultLabel(LOCTEXT("TypeColumnLabel", "Type"))
@@ -122,6 +136,10 @@ namespace UE::ConcertClientSharedSlate::ReplicationColumns::Property
 				{
 					InOutSearchStrings.Add(DisplayUtils::GetPropertyDisplayText(ObjectData.GetProperty()).ToString());
 				})
+				.IsLessThan_Lambda([](const FReplicatedPropertyData& Left, const FReplicatedPropertyData& Right)
+				{
+					return DisplayUtils::GetPropertyDisplayString(Left.GetProperty()) < DisplayUtils::GetPropertyDisplayString(Right.GetProperty());
+				})
 				.ColumnSortOrder(static_cast<int32>(EReplicationPropertyColumnOrder::Label)),
 			SHeaderRow::Column(LabelColumnId)
 				.DefaultLabel(LOCTEXT("LabelColumnLabel", "Label"))
@@ -150,6 +168,10 @@ namespace UE::ConcertClientSharedSlate::ReplicationColumns::Property
 				{
 					InOutSearchStrings.Add(GetDisplayText(ObjectData).ToString());
 				})
+				.IsLessThan_Lambda([](const FReplicatedPropertyData& Left, const FReplicatedPropertyData& Right)
+				{
+					return GetDisplayText(Left).ToString() < GetDisplayText(Right).ToString();
+				})
 				.ColumnSortOrder(static_cast<int32>(EReplicationPropertyColumnOrder::Type)),
 			SHeaderRow::Column(TypeColumnId)
 				.DefaultLabel(LOCTEXT("TypeColumnLabel", "Type"))
@@ -289,27 +311,6 @@ namespace UE::ConcertClientSharedSlate::ReplicationColumns::Property
 		}
 		return CheckBoxState;
 	}
-	
-	bool SortBySelectionThenByName_PropertyPredicate(
-		const TArray<FSoftObjectPath>& SelectedObjects,
-		const IReplicationStreamModel& Model,
-		const FReplicatedPropertyData& Left,
-		const FReplicatedPropertyData& Right
-		)
-	{
-		const ECheckBoxState LeftCheckboxState = GetPropertyCheckboxStateBasedOnSelection(Left.GetProperty(), SelectedObjects, Model);
-		const ECheckBoxState RightCheckboxState = GetPropertyCheckboxStateBasedOnSelection(Right.GetProperty(), SelectedObjects, Model);
-		
-		// Secondary sort by name
-		if (LeftCheckboxState == RightCheckboxState)
-		{
-			return DisplayUtils::GetPropertyDisplayString(Left.GetProperty()) < DisplayUtils::GetPropertyDisplayString(Right.GetProperty());
-		}
-
-		// Selected properties should appear first
-		return LeftCheckboxState == ECheckBoxState::Checked
-			&& (RightCheckboxState == ECheckBoxState::Unchecked || RightCheckboxState == ECheckBoxState::Undetermined);
-	}
 }
 
 #undef LOCTEXT_NAMESPACE
\ No newline at end of file
diff --git a/Engine/Plugins/Developer/Concert/ConcertUI/ConcertClientSharedSlate/Source/ConcertClientSharedSlate/Private/Replication/Editor/View/Tree/SelectionViewerColumns.h b/Engine/Plugins/Developer/Concert/ConcertUI/ConcertClientSharedSlate/Source/ConcertClientSharedSlate/Private/Replication/Editor/View/Tree/SelectionViewerColumns.h
--- a/Engine/Plugins/Developer/Concert/ConcertUI/ConcertClientSharedSlate/Source/ConcertClientSharedSlate/Private/Replication/Editor/View/Tree/SelectionViewerColumns.h
+++ b/Engine/Plugins/Developer/Concert/ConcertUI/ConcertClientSharedSlate/Source/ConcertClientSharedSlate/Private/Replication/Editor/View/Tree/SelectionViewerColumns.h
@@ -81,11 +81,4 @@ namespace UE::ConcertClientSharedSlate::ReplicationColumns::Property
 		TConstArrayView<FSoftObjectPath> Selection,
 		const IReplicationStreamModel& Model
 		);
-	/** Util for sorting property data by whether its ReplicatesColumns() column is checked. */
-	CONCERTCLIENTSHAREDSLATE_API bool SortBySelectionThenByName_PropertyPredicate(
-		const TArray<FSoftObjectPath>& SelectedObjects,
-		const IReplicationStreamModel& Model,
-		const FReplicatedPropertyData& Left,
-		const FReplicatedPropertyData& Right
-		);
 }
diff --git a/Engine/Plugins/Developer/Concert/ConcertUI/ConcertClientSharedSlate/Source/ConcertClientSharedSlate/Private/Replication/ReplicationWidgetFactories.cpp b/Engine/Plugins/Developer/Concert/ConcertUI/ConcertClientSharedSlate/Source/ConcertClientSharedSlate/Private/Replication/ReplicationWidgetFactories.cpp
--- a/Engine/Plugins/Developer/Concert/ConcertUI/ConcertClientSharedSlate/Source/ConcertClientSharedSlate/Private/Replication/ReplicationWidgetFactories.cpp
+++ b/Engine/Plugins/Developer/Concert/ConcertUI/ConcertClientSharedSlate/Source/ConcertClientSharedSlate/Private/Replication/ReplicationWidgetFactories.cpp
@@ -19,10 +19,13 @@ namespace UE::ConcertClientSharedSlate
 	{
 		return SNew(SBaseReplicationStreamEditor, Params.DataModel, Params.ObjectSource, Params.PropertySource)
 			.AdditionalObjectColumns(Params.ViewerParams.AdditionalObjectColumns)
+			.PrimaryObjectSort(Params.ViewerParams.PrimaryObjectSort)
+			.SecondaryObjectSort(Params.ViewerParams.SecondaryObjectSort)
 			.AdditionalPropertyColumns(Params.ViewerParams.AdditionalPropertyColumns)
+			.PrimaryPropertySort(Params.ViewerParams.PrimaryPropertySort)
+			.SecondaryPropertySort(Params.ViewerParams.SecondaryPropertySort)
 			.SubobjectModel(Params.ViewerParams.SubobjectModel)
 			.OnExtendObjectsContextMenu(Params.ViewerParams.OnExtendObjectsContextMenu)
-			.SortPropertyRowPredicate(Params.ViewerParams.SortPropertyRowPredicate)
 			.LeftOfObjectSearchBar() [ Params.ViewerParams.LeftOfObjectSearchBar.Widget ]
 			.LeftOfPropertySearchBar() [ Params.ViewerParams.LeftOfPropertySearchBar.Widget ]
 			.IsEditingEnabled(Params.IsEditingEnabled)
@@ -33,10 +36,13 @@ namespace UE::ConcertClientSharedSlate
 	{
 		return SNew(SDefaultReplicationStreamEditor, Params.DataModel, Params.ObjectSource, Params.PropertySource)
 			.AdditionalObjectColumns(Params.ViewerParams.AdditionalObjectColumns)
+			.PrimaryObjectSort(Params.ViewerParams.PrimaryObjectSort)
+			.SecondaryObjectSort(Params.ViewerParams.SecondaryObjectSort)
 			.AdditionalPropertyColumns(Params.ViewerParams.AdditionalPropertyColumns)
+			.PrimaryPropertySort(Params.ViewerParams.PrimaryPropertySort)
+			.SecondaryPropertySort(Params.ViewerParams.SecondaryPropertySort)
 			.SubobjectModel(Params.ViewerParams.SubobjectModel)
 			.OnExtendObjectsContextMenu(Params.ViewerParams.OnExtendObjectsContextMenu)
-			.SortPropertyRowPredicate(Params.ViewerParams.SortPropertyRowPredicate)
 			.LeftOfObjectSearchBar() [ Params.ViewerParams.LeftOfObjectSearchBar.Widget ]
 			.LeftOfPropertySearchBar() [ Params.ViewerParams.LeftOfPropertySearchBar.Widget ]
 			.IsEditingEnabled(Params.IsEditingEnabled)
diff --git a/Engine/Plugins/Developer/Concert/ConcertUI/ConcertClientSharedSlate/Source/ConcertClientSharedSlate/Private/Widgets/ClientName/SHorizontalClientList.cpp b/Engine/Plugins/Developer/Concert/ConcertUI/ConcertClientSharedSlate/Source/ConcertClientSharedSlate/Private/Widgets/ClientName/SHorizontalClientList.cpp
--- a/Engine/Plugins/Developer/Concert/ConcertUI/ConcertClientSharedSlate/Source/ConcertClientSharedSlate/Private/Widgets/ClientName/SHorizontalClientList.cpp
+++ b/Engine/Plugins/Developer/Concert/ConcertUI/ConcertClientSharedSlate/Source/ConcertClientSharedSlate/Private/Widgets/ClientName/SHorizontalClientList.cpp
@@ -3,6 +3,7 @@
 #include "Widgets/ClientName/SHorizontalClientList.h"
 
 #include "IConcertClient.h"
+#include "Widgets/ClientName/SClientName.h"
 
 #include "Widgets/ClientName/SLocalClientName.h"
 #include "Widgets/ClientName/SRemoteClientName.h"
@@ -14,6 +15,45 @@
 
 namespace UE::ConcertClientSharedSlate
 {
+	namespace HorizontalClientList
+	{
+		TArray<FConcertSessionClientInfo> GetSortedClients(
+			const IConcertClient& LocalConcertClient,
+			const TConstArrayView<FGuid>& Clients,
+			const SHorizontalClientList::FSortPredicate& SortPredicate
+			)
+		{
+			TArray<FConcertSessionClientInfo> ClientsToDisplay;
+			
+			const TSharedPtr<IConcertClientSession> ClientSession = LocalConcertClient.GetCurrentSession();
+			if (!ensureMsgf(ClientSession, TEXT("This widget does not work if the local client is not in a session!")))
+			{
+				return ClientsToDisplay;
+			}
+		
+			// Prefetch the client info to avoid many FindSessionClient during Sort()
+			for (const FGuid& Client : Clients)
+			{
+				FConcertSessionClientInfo Info;
+				if (ClientSession->FindSessionClient(Client, Info))
+				{
+					ClientsToDisplay.Add(Info);
+				}
+				else if (Client == ClientSession->GetSessionClientEndpointId())
+				{
+					ClientsToDisplay.Add({Client, ClientSession->GetLocalClientInfo()});
+				}
+			}
+
+			ClientsToDisplay.Sort([&SortPredicate](const FConcertSessionClientInfo& Left, const FConcertSessionClientInfo& Right)
+			{
+				return SortPredicate.Execute(Left, Right);
+			});
+			
+			return ClientsToDisplay;
+		}
+	}
+	
 	bool SHorizontalClientList::SortLocalClientFirstThenAlphabetical(const FConcertSessionClientInfo& Left, const FConcertSessionClientInfo& Right, TSharedRef<IConcertClient> Client)
 	{
 		// If one of the compare clients is local, always return that the local client is smaller.
@@ -27,6 +67,23 @@ namespace UE::ConcertClientSharedSlate
 			|| (!bRightIsLocalClient && Left.ClientInfo.DisplayName < Right.ClientInfo.DisplayName);
 	}
 
+	TOptional<FString> SHorizontalClientList::GetDisplayString(const IConcertClient& LocalConcertClient, const TConstArrayView<FGuid>& Clients, const FSortPredicate& SortPredicate)
+	{
+		const TSharedPtr<IConcertClientSession> ClientSession = LocalConcertClient.GetCurrentSession();
+		const TArray<FConcertSessionClientInfo> ClientsToDisplay = HorizontalClientList::GetSortedClients(LocalConcertClient, Clients, SortPredicate);
+		if (!ClientSession || ClientsToDisplay.IsEmpty())
+		{
+			return {};
+		}
+		
+		return FString::JoinBy(ClientsToDisplay, TEXT(", "), [&ClientSession](const FConcertSessionClientInfo& ClientInfo)
+			{
+				// GetSortedClients should return empty if GetSortedClients is invalid
+				const bool bIsLocalClient = ensure(ClientSession) && ClientInfo.ClientEndpointId == ClientSession->GetSessionClientEndpointId();
+				return SClientName::GetDisplayText(ClientInfo.ClientInfo, bIsLocalClient).ToString();
+			});
+	}
+
 	void SHorizontalClientList::Construct(const FArguments& InArgs, TSharedRef<IConcertClient> InClient)
 	{
 		LocalConcertClient = MoveTemp(InClient);
@@ -64,30 +121,8 @@ namespace UE::ConcertClientSharedSlate
 		WidgetSwitcher->SetActiveWidgetIndex(1);
 
 		const TSharedPtr<IConcertClientSession> ClientSession = LocalConcertClient->GetCurrentSession();
-		if (!ensureMsgf(ClientSession, TEXT("This widget does not work if the local client is not in a session!")))
-		{
-			return;
-		}
-
-		// Prefetch the client info to avoid many FindSessionClient during Sort()
-		TArray<FConcertSessionClientInfo> ClientsToDisplay;
-		for (const FGuid& Client : Clients)
-		{
-			FConcertSessionClientInfo Info;
-			if (ClientSession->FindSessionClient(Client, Info))
-			{
-				ClientsToDisplay.Add(Info);
-			}
-			else if (Client == ClientSession->GetSessionClientEndpointId())
-			{
-				ClientsToDisplay.Add({Client, ClientSession->GetLocalClientInfo()});
-			}
-		}
-
-		ClientsToDisplay.Sort([this](const FConcertSessionClientInfo& Left, const FConcertSessionClientInfo& Right)
-		{
-			return SortPredicateDelegate.Execute(Left, Right);
-		});
+		const TArray<FConcertSessionClientInfo> ClientsToDisplay = HorizontalClientList::GetSortedClients(*LocalConcertClient, Clients, SortPredicateDelegate);
+		
 		bool bIsFirst = true;
 		for (const FConcertSessionClientInfo& Info : ClientsToDisplay)
 		{
@@ -100,8 +135,9 @@ namespace UE::ConcertClientSharedSlate
 					.Font(NameFont)
 				];
 			}
-			
-			if (Info.ClientEndpointId == ClientSession->GetSessionClientEndpointId())
+
+			// GetSortedClients should return empty if GetSortedClients is invalid
+			if (ensure(ClientSession) && Info.ClientEndpointId == ClientSession->GetSessionClientEndpointId())
 			{
 				ScrollBox->AddSlot()
 				[
diff --git a/Engine/Plugins/Developer/Concert/ConcertUI/ConcertClientSharedSlate/Source/ConcertClientSharedSlate/Public/Replication/Editor/View/IReplicationStreamViewer.h b/Engine/Plugins/Developer/Concert/ConcertUI/ConcertClientSharedSlate/Source/ConcertClientSharedSlate/Public/Replication/Editor/View/IReplicationStreamViewer.h
--- a/Engine/Plugins/Developer/Concert/ConcertUI/ConcertClientSharedSlate/Source/ConcertClientSharedSlate/Public/Replication/Editor/View/IReplicationStreamViewer.h
+++ b/Engine/Plugins/Developer/Concert/ConcertUI/ConcertClientSharedSlate/Source/ConcertClientSharedSlate/Public/Replication/Editor/View/IReplicationStreamViewer.h
@@ -18,8 +18,17 @@ namespace UE::ConcertClientSharedSlate
 		/** Call after the data underlying the model was externally changed and needs to be redisplayed in the UI. */
 		virtual void Refresh() = 0;
 
-		/** @return The objects selected in the top of the viewer; these are the objects that IReplicationSubobjectView bases its view of. */
-		virtual TArray<FSoftObjectPath> GetSelectedTopLevelObjects() const = 0;
+		/**
+		 * Requests that column be resorted; the column is in the top object view.
+		 * This is to be called in response to a column's content changing. The rows will be resorted if the given column has a sort priority assigned.
+		 */
+		virtual void RequestObjectColumnResort(const FName& ColumnId) = 0;
+		/**
+		 * Requests that column be resorted; the column is in the bottom property view.
+		 * This is to be called in response to a column's content changing. The rows will be resorted if the given column has a sort priority assigned.
+		 */
+		virtual void RequestPropertyColumnResort(const FName& ColumnId) = 0;
+
 		/**
 		 * @return The objects for which the properties are being edited / displayed.
 		 * If there is an IReplicationSubobjectView, this is IReplicationSubobjectView::GetSelectedObjects.
diff --git a/Engine/Plugins/Developer/Concert/ConcertUI/ConcertClientSharedSlate/Source/ConcertClientSharedSlate/Public/Replication/Editor/View/ReplicationColumn.h b/Engine/Plugins/Developer/Concert/ConcertUI/ConcertClientSharedSlate/Source/ConcertClientSharedSlate/Public/Replication/Editor/View/ReplicationColumn.h
--- a/Engine/Plugins/Developer/Concert/ConcertUI/ConcertClientSharedSlate/Source/ConcertClientSharedSlate/Public/Replication/Editor/View/ReplicationColumn.h
+++ b/Engine/Plugins/Developer/Concert/ConcertUI/ConcertClientSharedSlate/Source/ConcertClientSharedSlate/Public/Replication/Editor/View/ReplicationColumn.h
@@ -22,11 +22,12 @@ namespace UE::ConcertClientSharedSlate
 		DECLARE_DELEGATE_RetVal_OneParam(TSharedRef<SWidget>, FGenerateColumnWidget,
 			const FBuildArgs& InArgs
 			);
-
 		DECLARE_DELEGATE_TwoParams(FPopulateSearchString,
 			const TListItemType& InArgs,
 			TArray<FString>& InOutSearchStrings
 			);
+		/** Optional Whether Left < Right for this column's displayed content. Used for sorting. */
+		DECLARE_DELEGATE_RetVal_TwoParams(bool, FIsLessThan, const TListItemType& Left, const TListItemType& Right);
 		
 		SLATE_BEGIN_ARGS(TReplicationColumn)
 		{}
@@ -34,6 +35,8 @@ namespace UE::ConcertClientSharedSlate
 			SLATE_EVENT(FGenerateColumnWidget, GenerateWidgetColumn)
 			/** Callback used to generate search items for this column */
 			SLATE_EVENT(FPopulateSearchString, PopulateSearchItems)
+			/** Optional Whether Left <= Right for this column's displayed content. Used for sorting. */
+			SLATE_EVENT(FIsLessThan, IsLessThan)
 			/** Where in the row this column is found with respect to the other columns */
 			SLATE_ARGUMENT(int32, ColumnSortOrder)
 		SLATE_END_ARGS()
@@ -42,6 +45,7 @@ namespace UE::ConcertClientSharedSlate
 			: FColumn(InColumnArgs)
 			, GenerateColumnWidgetCallback(InArgs._GenerateWidgetColumn)
 			, PopulateSearchStringCallback(InArgs._PopulateSearchItems)
+			, IsLessThanCallback(InArgs._IsLessThan)
 			, ColumnSortOrderValue(InArgs._ColumnSortOrder)
 		{}
 		
@@ -49,6 +53,7 @@ namespace UE::ConcertClientSharedSlate
 			: FColumn(Column)
 			, GenerateColumnWidgetCallback(InArgs._GenerateWidgetColumn)
 			, PopulateSearchStringCallback(InArgs._PopulateSearchItems)
+			, IsLessThanCallback(InArgs._IsLessThan)
 			, ColumnSortOrderValue(InArgs._ColumnSortOrder)
 		{}
 
@@ -80,30 +85,54 @@ namespace UE::ConcertClientSharedSlate
 		TReplicationColumn<TOtherColumnType> TransformColumn(TTransformOp TransformOperation = [](const TOtherColumnType& RowData) -> TListItemType { return RowData; }) const
 		{
 			using TReturnColumnType = TReplicationColumn<TOtherColumnType>;
-			FGenerateColumnWidget GenerateWidget = GenerateColumnWidgetCallback;
-			FPopulateSearchString PopulateSearchString = PopulateSearchStringCallback;
+			using TTargetPopulateSearchString = typename TReplicationColumn<TOtherColumnType>::FPopulateSearchString;
+			using TTargetLessThan = typename TReplicationColumn<TOtherColumnType>::FIsLessThan;
+			
+			FGenerateColumnWidget ThisGenerateWidget = GenerateColumnWidgetCallback;
+			FPopulateSearchString ThisPopulateSearchString = PopulateSearchStringCallback;
+			FIsLessThan ThisIsLessThan = IsLessThanCallback;
+
+			auto PopulateSearchItemsLambda = [TransformOperation, ThisPopulateSearchString](const TOtherColumnType& InOtherRowData, TArray<FString>& InOutSearchStrings)
+			{
+				ThisPopulateSearchString.Execute(TransformOperation(InOtherRowData), InOutSearchStrings);
+			};
+			auto IsLessThanLambda = [TransformOperation, ThisIsLessThan](const TOtherColumnType& Left, const TOtherColumnType& Right)
+			{
+				return ThisIsLessThan.Execute(TransformOperation(Left), TransformOperation(Right));
+			};
+			
+			TTargetPopulateSearchString TargetPopulateSearchString = ThisPopulateSearchString.IsBound()
+				? TTargetPopulateSearchString::CreateLambda(PopulateSearchItemsLambda)
+				: TTargetPopulateSearchString{};
+			TTargetLessThan TargetLessThan = ThisIsLessThan.IsBound()
+				? TTargetLessThan::CreateLambda(IsLessThanLambda)
+				: TTargetLessThan{};
+			
 			return TReturnColumnType(
 				typename TReturnColumnType::FArguments()
-					.GenerateWidgetColumn_Lambda([TransformOperation, GenerateWidget](const typename TReturnColumnType::FBuildArgs& BuildArgs)
+					.GenerateWidgetColumn_Lambda([TransformOperation, ThisGenerateWidget](const typename TReturnColumnType::FBuildArgs& BuildArgs)
 					{
-						return GenerateWidget.Execute({ BuildArgs.HighlightText, TransformOperation(BuildArgs.RowData)});
-					})
-					.PopulateSearchItems_Lambda([TransformOperation, PopulateSearchString](const TOtherColumnType& InOtherRowData, TArray<FString>& InOutSearchStrings)
-					{
-						if (PopulateSearchString.IsBound())
-						{
-							PopulateSearchString.Execute(TransformOperation(InOtherRowData), InOutSearchStrings);
-						}
+						return ThisGenerateWidget.Execute({ BuildArgs.HighlightText, TransformOperation(BuildArgs.RowData)});
 					})
+					.PopulateSearchItems(TargetPopulateSearchString)
+					.IsLessThan(TargetLessThan)
 					.ColumnSortOrder(ColumnSortOrderValue),
 				*this
 			);
 		}
+
+		bool CanBeSorted() const { return IsLessThanCallback.IsBound(); }
+		bool IsLessThan(const TListItemType& Left, const TListItemType& Right) const
+		{
+			return ensure(IsLessThanCallback.IsBound())
+				&& IsLessThanCallback.Execute(Left, Right);
+		}
 		
 	private:
 		
 		FGenerateColumnWidget GenerateColumnWidgetCallback;
 		FPopulateSearchString PopulateSearchStringCallback;
+		FIsLessThan IsLessThanCallback;
 		/** Determines whether this column is the first, etc. */
 		int32 ColumnSortOrderValue = 0;
 	};
diff --git a/Engine/Plugins/Developer/Concert/ConcertUI/ConcertClientSharedSlate/Source/ConcertClientSharedSlate/Public/Replication/Editor/View/ReplicationColumnsUtils.h b/Engine/Plugins/Developer/Concert/ConcertUI/ConcertClientSharedSlate/Source/ConcertClientSharedSlate/Public/Replication/Editor/View/ReplicationColumnsUtils.h
--- a/Engine/Plugins/Developer/Concert/ConcertUI/ConcertClientSharedSlate/Source/ConcertClientSharedSlate/Public/Replication/Editor/View/ReplicationColumnsUtils.h
+++ b/Engine/Plugins/Developer/Concert/ConcertUI/ConcertClientSharedSlate/Source/ConcertClientSharedSlate/Public/Replication/Editor/View/ReplicationColumnsUtils.h
@@ -18,6 +18,15 @@ namespace UE::ConcertClientSharedSlate
 
 namespace UE::ConcertClientSharedSlate
 {
+	/** Info about a column that is being sorted. */
+	struct FColumnSortInfo
+	{
+		FName SortedColumnId;
+		EColumnSortMode::Type SortMode = EColumnSortMode::None;
+
+		bool IsValid() const { return SortedColumnId != NAME_None && SortMode != EColumnSortMode::None; }
+	};
+	
 	template<typename TListItemType>
 	struct TReplicationColumnDelegates
 	{
@@ -95,6 +104,13 @@ namespace UE::ConcertClientSharedSlate
 								OnChangedDelegate.Execute(bIsChecked, RowData);
 							});
 					})
+					.IsLessThan_Lambda([Delegates](const TListItemType& Left, const TListItemType& Right)
+					{
+						const bool bIsLeftChecked = Delegates.GetCheckboxStateDelegate.IsBound() && Delegates.GetCheckboxStateDelegate.Execute(Left) == ECheckBoxState::Checked;
+						const bool bIsRightChecked = Delegates.GetCheckboxStateDelegate.IsBound() && Delegates.GetCheckboxStateDelegate.Execute(Right) == ECheckBoxState::Checked;
+						// Less for checkbox means: checked < unchecked. !bIsRightChecked because checked == checked so it cannot be less.
+						return bIsLeftChecked && !bIsRightChecked;
+					})
 				.ColumnSortOrder(Priority),
 			SHeaderRow::Column(ColumnId)
 				.DefaultLabel(DefaultLabel)
diff --git a/Engine/Plugins/Developer/Concert/ConcertUI/ConcertClientSharedSlate/Source/ConcertClientSharedSlate/Public/Replication/ReplicationWidgetFactories.h b/Engine/Plugins/Developer/Concert/ConcertUI/ConcertClientSharedSlate/Source/ConcertClientSharedSlate/Public/Replication/ReplicationWidgetFactories.h
--- a/Engine/Plugins/Developer/Concert/ConcertUI/ConcertClientSharedSlate/Source/ConcertClientSharedSlate/Public/Replication/ReplicationWidgetFactories.h
+++ b/Engine/Plugins/Developer/Concert/ConcertUI/ConcertClientSharedSlate/Source/ConcertClientSharedSlate/Public/Replication/ReplicationWidgetFactories.h
@@ -6,6 +6,7 @@
 #include "ReplicationWidgetDelegates.h"
 
 #include "Delegates/Delegate.h"
+#include "Editor/View/ReplicationColumnsUtils.h"
 #include "Misc/Attribute.h"
 #include "Templates/SharedPointer.h"
 
@@ -63,13 +64,20 @@ namespace UE::ConcertClientSharedSlate
 
 		/** Called to generate the context menu for objects. This extends the options already generated by this widget and is called at the end. */
 		FExtendMenu OnExtendObjectsContextMenu;
-		/** Optional. Used for determining the order in which properties are displayed. */
-		FSortPropertyPredicate SortPropertyRowPredicate;
 		
 		/** Additional columns to add to the object view */
 		TArray<ReplicationColumns::FReplicationTopLevelObjectColumn> AdditionalObjectColumns;
+		/** Optional initial primary sort mode for object rows */
+		FColumnSortInfo PrimaryObjectSort;
+		/** Optional initial secondary sort mode for object rows */
+		FColumnSortInfo SecondaryObjectSort;
+		
 		/** Additional columns to add to the property view */
 		TArray<ReplicationColumns::FReplicationPropertyColumn> AdditionalPropertyColumns;
+		/** Optional initial primary sort mode for object rows */
+		FColumnSortInfo PrimaryPropertySort;
+		/** Optional initial secondary sort mode for object rows */
+		FColumnSortInfo SecondaryPropertySort;
 		
 		/** Optional widget to add to the left of the object list search bar. */
 		TAlwaysValidWidget LeftOfObjectSearchBar;
diff --git a/Engine/Plugins/Developer/Concert/ConcertUI/ConcertClientSharedSlate/Source/ConcertClientSharedSlate/Public/Widgets/ClientName/SHorizontalClientList.h b/Engine/Plugins/Developer/Concert/ConcertUI/ConcertClientSharedSlate/Source/ConcertClientSharedSlate/Public/Widgets/ClientName/SHorizontalClientList.h
--- a/Engine/Plugins/Developer/Concert/ConcertUI/ConcertClientSharedSlate/Source/ConcertClientSharedSlate/Public/Widgets/ClientName/SHorizontalClientList.h
+++ b/Engine/Plugins/Developer/Concert/ConcertUI/ConcertClientSharedSlate/Source/ConcertClientSharedSlate/Public/Widgets/ClientName/SHorizontalClientList.h
@@ -22,6 +22,9 @@ namespace UE::ConcertClientSharedSlate
 
 		DECLARE_DELEGATE_RetVal_TwoParams(bool, FSortPredicate, const FConcertSessionClientInfo& Left, const FConcertSessionClientInfo& FConcertSessionClientInfo);
 		static bool SortLocalClientFirstThenAlphabetical(const FConcertSessionClientInfo& Left, const FConcertSessionClientInfo& Right, TSharedRef<IConcertClient> Client);
+
+		/** @return The display string a SHorizontalClientList would display with the given state. Returns unset optional if EmptyListSlot would be shown. */
+		static TOptional<FString> GetDisplayString(const IConcertClient& LocalConcertClient, const TConstArrayView<FGuid>& Clients, const FSortPredicate& SortPredicate);
 		
 		SLATE_BEGIN_ARGS(SHorizontalClientList)
 			: _Font(FAppStyle::Get().GetFontStyle("NormalFont"))
