Refactor stream extender API so it can be exposed as a more simple Blueprint API

diff --git a/Engine/Plugins/Developer/Concert/ConcertApp/MultiUserClient/Source/MultiUserClient/Private/Replication/Client/ReplicationClient.cpp b/Engine/Plugins/Developer/Concert/ConcertApp/MultiUserClient/Source/MultiUserClient/Private/Replication/Client/ReplicationClient.cpp
--- a/Engine/Plugins/Developer/Concert/ConcertApp/MultiUserClient/Source/MultiUserClient/Private/Replication/Client/ReplicationClient.cpp
+++ b/Engine/Plugins/Developer/Concert/ConcertApp/MultiUserClient/Source/MultiUserClient/Private/Replication/Client/ReplicationClient.cpp
@@ -2,11 +2,10 @@
 
 #include "ReplicationClient.h"
 
-#include "MultiUserReplicationSettings.h"
 #include "Assets/MultiUserReplicationClientPreset.h"
 #include "Replication/Editor/Model/IEditableReplicationStreamModel.h"
-#include "Replication/Editor/Model/StreamExtenderBySettings.h"
 #include "Replication/ReplicationWidgetFactories.h"
+#include "Replication/Stream/Discovery/MultiUserStreamExtender.h"
 #include "Replication/Stream/StreamChangeTracker.h"
 #include "Replication/Submission/Data/AuthoritySubmission.h"
 
@@ -28,16 +27,11 @@ namespace UE::MultiUserClient
 		, AuthoritySynchronizer(MoveTemp(InAuthoritySynchronizer))
 		, SubmissionWorkflow(MoveTemp(InSubmissionWorkflow))
 		, SubmissionQueue(*SubmissionWorkflow)
-		, LocalClientEditModel(ConcertClientSharedSlate::CreatePropertySelectionModel(
+		, LocalClientEditModel(CreatePropertySelectionModel(
 			*ClientContentStorage->Stream,
 			ClientContentStorage->Stream->MakeReplicationMapGetterAttribute(),
-			// Use MU settings for auto adding properties & objects
-			MakeShared<ConcertClientSharedSlate::FStreamExtenderBySettings>(
-				TAttribute<const FConcertReplicationEditorSettings*>::CreateLambda([]()
-				{
-					return &UMultiUserReplicationSettings::Get()->ReplicationEditorSettings;
-				}))
-			))
+			MakeShared<FMultiUserStreamExtender>())
+			)
 		, LocalClientStreamDiffer(
 			GetStreamSynchronizer(),
 			ClientContentStorage->Stream->MakeReplicationMapGetterAttribute(),
diff --git a/Engine/Plugins/Developer/Concert/ConcertApp/MultiUserClient/Source/MultiUserClient/Private/Replication/Stream/Discovery/MultiUserStreamExtender.cpp b/Engine/Plugins/Developer/Concert/ConcertApp/MultiUserClient/Source/MultiUserClient/Private/Replication/Stream/Discovery/MultiUserStreamExtender.cpp
new file mode 100644
--- /dev/null
+++ b/Engine/Plugins/Developer/Concert/ConcertApp/MultiUserClient/Source/MultiUserClient/Private/Replication/Stream/Discovery/MultiUserStreamExtender.cpp
@@ -0,0 +1,22 @@
+// Copyright Epic Games, Inc. All Rights Reserved.
+
+#include "MultiUserStreamExtender.h"
+
+#include "MultiUserReplicationSettings.h"
+
+namespace UE::MultiUserClient
+{
+	FMultiUserStreamExtender::FMultiUserStreamExtender()
+		: ExtendBySettings(
+			// Use MU settings for auto adding properties & objects
+			TAttribute<const FConcertReplicationEditorSettings*>::CreateLambda([]()
+			{
+				return &UMultiUserReplicationSettings::Get()->ReplicationEditorSettings;
+			}))
+	{}
+
+	void FMultiUserStreamExtender::ExtendStream(UObject& ExtendedObject, ConcertClientSharedSlate::IStreamExtensionContext& Context)
+	{
+		ExtendBySettings.ExtendStream(ExtendedObject, Context);
+	}
+}
diff --git a/Engine/Plugins/Developer/Concert/ConcertApp/MultiUserClient/Source/MultiUserClient/Private/Replication/Stream/Discovery/MultiUserStreamExtender.h b/Engine/Plugins/Developer/Concert/ConcertApp/MultiUserClient/Source/MultiUserClient/Private/Replication/Stream/Discovery/MultiUserStreamExtender.h
new file mode 100644
--- /dev/null
+++ b/Engine/Plugins/Developer/Concert/ConcertApp/MultiUserClient/Source/MultiUserClient/Private/Replication/Stream/Discovery/MultiUserStreamExtender.h
@@ -0,0 +1,34 @@
+// Copyright Epic Games, Inc. All Rights Reserved.
+
+#pragma once
+
+#include "Replication/Editor/Model/Extension/IStreamExtender.h"
+#include "Replication/Editor/Model/StreamExtenderBySettings.h"
+
+namespace UE::MultiUserClient
+{
+	class FReplicationDiscoveryContainer;
+	
+	/**
+	 * When the user adds an object, this object handles auto selecting properties and adding additional objects from context.
+	 * The following sources exist:
+	 * - Static Settings: user can specify properties & objects in the MU project settings
+	 * - Dynamic API: by using the IMultiUserReplication, external modules can register dynamic rules for auto-discovery.
+	 */
+	class FMultiUserStreamExtender : public ConcertClientSharedSlate::IStreamExtender
+	{
+	public:
+
+		FMultiUserStreamExtender();
+
+		//~ Begin IStreamExtender Interface
+		virtual void ExtendStream(UObject& ExtendedObject, ConcertClientSharedSlate::IStreamExtensionContext& Context) override;
+		//~ End IStreamExtender Interface
+
+	private:
+
+		/** Handles properties from the MU settings. */
+		ConcertClientSharedSlate::FStreamExtenderBySettings ExtendBySettings;
+	};
+}
+
diff --git a/Engine/Plugins/Developer/Concert/ConcertSync/ConcertSyncCore/Source/ConcertSyncCore/Public/Replication/Data/ConcertPropertySelection.h b/Engine/Plugins/Developer/Concert/ConcertSync/ConcertSyncCore/Source/ConcertSyncCore/Public/Replication/Data/ConcertPropertySelection.h
--- a/Engine/Plugins/Developer/Concert/ConcertSync/ConcertSyncCore/Source/ConcertSyncCore/Public/Replication/Data/ConcertPropertySelection.h
+++ b/Engine/Plugins/Developer/Concert/ConcertSync/ConcertSyncCore/Source/ConcertSyncCore/Public/Replication/Data/ConcertPropertySelection.h
@@ -21,12 +21,17 @@ struct CONCERTSYNCCORE_API FConcertPropertyChain
 	/**
 	 * The name of inner properties of containers (array, set, map), which are either primitive or native serialized structs.
 	 * 
-	 * Note that inner properties are only included if they are at the end of the path (and only for primitive or native serialized structs!), e.g.
+	 * Inner properties are included if they are at the end of the path (and only for primitive or native serialized structs!).
+	 * Cases:
 	 * - ArrayOfStructs.ArrayOfFloats.Value
-	 *	This means there is an array property called ArrayOfStructs containing structs. The contained struct has an array of floats property
+	 *	There is an array property called ArrayOfStructs containing structs. The contained struct has an array of floats property
 	 *	called ArrayOfFloats. Value corresponds to FArrayProperty::Inner and is called InternalContainerPropertyValueName.
-	 *	
-	 * - but not ArrayOfStructs.Value.ArrayOfFloats.Value.
+	 * - ArrayOfStructs.ArrayOfFloats
+	 *  Similar situation but only the size of the array will be replicated. The elements are default initialized.
+	 * - ArrayOfStructs.Value
+	 *  There is an array property called ArrayOfStructs and the struct has a native Serialize() function.
+	 *  
+	 * - BUT not ArrayOfStructs.Value.ArrayOfFloats.Value.
 	 *	This path would imply a different situation in which ArrayOfStructs contains a struct with a struct property called Value;
 	 *	Value's owning struct property would contain an array of floats called ArrayOfFloats.
 	 */
diff --git a/Engine/Plugins/Developer/Concert/ConcertUI/ConcertClientSharedSlate/Source/ConcertClientSharedSlate/Private/Replication/Editor/Model/GenericReplicationStreamModel.cpp b/Engine/Plugins/Developer/Concert/ConcertUI/ConcertClientSharedSlate/Source/ConcertClientSharedSlate/Private/Replication/Editor/Model/GenericReplicationStreamModel.cpp
--- a/Engine/Plugins/Developer/Concert/ConcertUI/ConcertClientSharedSlate/Source/ConcertClientSharedSlate/Private/Replication/Editor/Model/GenericReplicationStreamModel.cpp
+++ b/Engine/Plugins/Developer/Concert/ConcertUI/ConcertClientSharedSlate/Source/ConcertClientSharedSlate/Private/Replication/Editor/Model/GenericReplicationStreamModel.cpp
@@ -2,15 +2,35 @@
 
 #include "GenericReplicationStreamModel.h"
 
+#include "ConcertLogGlobal.h"
 #include "Replication/PropertyChainUtils.h"
 #include "Replication/Data/ObjectReplicationMap.h"
+#include "Replication/Editor/Model/Extension/IStreamExtender.h"
+#include "Replication/Editor/Model/Extension/IStreamExtensionContext.h"
 
 #include "Algo/AllOf.h"
 #include "Containers/Queue.h"
-#include "Replication/Editor/Model/IStreamExtender.h"
 
 namespace UE::ConcertClientSharedSlate
 {
+	namespace Private
+	{
+		/** @return Whether any property was added */
+		static bool AddParentProperties(const UStruct& Class, const FConcertPropertyChain& PropertyToAdd, TArray<FConcertPropertyChain>& ReplicatedProperties)
+		{
+			bool bAddedAtLeastOne = false;
+			ConcertSyncCore::PropertyChain::ForEachReplicatableConcertProperty(Class, [&bAddedAtLeastOne, &ReplicatedProperties, &PropertyToAdd](FConcertPropertyChain&& Property)
+				{
+					if (Property.IsParentOf(PropertyToAdd))
+					{
+						bAddedAtLeastOne |= ReplicatedProperties.AddUnique(Property) != INDEX_NONE;
+					}
+					return EBreakBehavior::Continue;
+				});
+			return bAddedAtLeastOne;
+		}
+	}
+
 	FGenericReplicationStreamModel::FGenericReplicationStreamModel(
 		TAttribute<FObjectReplicationMap*> InReplicationMapAttribute,
 		TSharedPtr<IStreamExtender> InExtender
@@ -187,17 +207,11 @@ namespace UE::ConcertClientSharedSlate
 		for (const FConcertPropertyChain& AddedProperty : Properties)
 		{
 			bAddedAtLeastOne |= ReplicatedProperties.AddUnique(AddedProperty) != INDEX_NONE;
-
 			// Parent properties must also be added
 			// Not exactly efficient to iterate through the hierarchy for every removed item but it should be fine... Properties.Num() == 1 is the most common case
-			ConcertSyncCore::PropertyChain::ForEachReplicatableConcertProperty(*Class, [&bAddedAtLeastOne, &AssignedProperties, &AddedProperty](FConcertPropertyChain&& Property)
-			{
-				if (Property.IsParentOf(AddedProperty))
-				{
-					bAddedAtLeastOne |= AssignedProperties->PropertySelection.ReplicatedProperties.AddUnique(Property) != INDEX_NONE;
-				}
-				return EBreakBehavior::Continue;
-			});
+			bAddedAtLeastOne |= Private::AddParentProperties(*Class, AddedProperty, ReplicatedProperties);
+			
+			// TODO UE-202079: Make sure to append FConcertPropertyChain::InternalContainerPropertyValueName if needed
 		}
 
 		if (bAddedAtLeastOne)
@@ -257,28 +271,60 @@ namespace UE::ConcertClientSharedSlate
 			return;
 		}
 
-		TQueue<UObject*> ObjectsToProcess;
-		ObjectsToProcess.Enqueue(&AddedObject);
-
-		UObject* CurrentObject;
-		while (ObjectsToProcess.Dequeue(CurrentObject))
+		class FExtensionContext : public IStreamExtensionContext
 		{
-			FReplicatedObjectInfo& ObjectInfo = ReplicationMap.ReplicatedObjects[CurrentObject];
-			Extender->ExtendObjectProperties(*CurrentObject, [&ObjectInfo](FConcertPropertyChain PropertyChain)
+		public:
+
+			TQueue<UObject*> ObjectsToProcess;
+			FObjectReplicationMap& ReplicationMap;
+			TArray<UObject*>& ObjectsAddedSoFar;
+
+			explicit FExtensionContext(FObjectReplicationMap& ReplicationMap, TArray<UObject*>& ObjectsAddedSoFar)
+				: ReplicationMap(ReplicationMap)
+				, ObjectsAddedSoFar(ObjectsAddedSoFar)
+			{}
+
+			virtual void AddPropertyTo(UObject& Object, FConcertPropertyChain&& PropertyChain) override
 			{
-				ObjectInfo.PropertySelection.ReplicatedProperties.AddUnique(MoveTemp(PropertyChain));
-			});
+				AddAdditionalObject(Object);
+				FReplicatedObjectInfo& ObjectInfo = ReplicationMap.ReplicatedObjects[&Object];
+
+				constexpr bool bLog = false;
+				const FProperty* ResolvedProperty = PropertyChain.ResolveProperty(*Object.GetClass(), bLog);
+				if (!ResolvedProperty || !ConcertSyncCore::PropertyChain::IsReplicatableProperty(*ResolvedProperty))
+				{
+					UE_LOG(LogConcert, Warning, TEXT("Property \"%s\" is not a valid property to assign to object \"%s\"."), *PropertyChain.ToString(), *Object.GetPathName());
+					return;
+				}
 
-			Extender->AppendAdditionalObjects(*CurrentObject, [&ReplicationMap, &ObjectsAddedSoFar, &ObjectsToProcess](UObject& AdditionalObject)
+				if (!ObjectInfo.PropertySelection.ReplicatedProperties.Contains(PropertyChain))
+				{
+					TArray<FConcertPropertyChain>& Properties = ObjectInfo.PropertySelection.ReplicatedProperties;
+					Properties.Emplace(MoveTemp(PropertyChain));
+					Private::AddParentProperties(*Object.GetClass(), PropertyChain, Properties);
+					// TODO UE-202079: Make sure to append FConcertPropertyChain::InternalContainerPropertyValueName if needed
+				}
+			}
+			
+			virtual void AddAdditionalObject(UObject& Object) override
 			{
-				if (!ObjectsAddedSoFar.Contains(&AdditionalObject))
+				if (!ObjectsAddedSoFar.Contains(&Object))
 				{
-					ObjectsAddedSoFar.AddUnique(&AdditionalObject);
-					ReplicationMap.ReplicatedObjects.FindOrAdd(&AdditionalObject)
-						.ClassPath = AdditionalObject.GetClass();
-					ObjectsToProcess.Enqueue(&AdditionalObject);
+					ObjectsAddedSoFar.AddUnique(&Object);
+					ReplicationMap.ReplicatedObjects.FindOrAdd(&Object)
+						.ClassPath = Object.GetClass();
+					ObjectsToProcess.Enqueue(&Object);
 				}
-			});
+			}
+		};
+
+		FExtensionContext Context(ReplicationMap, ObjectsAddedSoFar);
+		Context.ObjectsToProcess.Enqueue(&AddedObject);
+		
+		UObject* CurrentObject;
+		while (Context.ObjectsToProcess.Dequeue(CurrentObject))
+		{
+			Extender->ExtendStream(*CurrentObject, Context);
 		}
 	}
 }
diff --git a/Engine/Plugins/Developer/Concert/ConcertUI/ConcertClientSharedSlate/Source/ConcertClientSharedSlate/Private/Replication/Editor/Model/StreamExtenderBySettings.cpp b/Engine/Plugins/Developer/Concert/ConcertUI/ConcertClientSharedSlate/Source/ConcertClientSharedSlate/Private/Replication/Editor/Model/StreamExtenderBySettings.cpp
--- a/Engine/Plugins/Developer/Concert/ConcertUI/ConcertClientSharedSlate/Source/ConcertClientSharedSlate/Private/Replication/Editor/Model/StreamExtenderBySettings.cpp
+++ b/Engine/Plugins/Developer/Concert/ConcertUI/ConcertClientSharedSlate/Source/ConcertClientSharedSlate/Private/Replication/Editor/Model/StreamExtenderBySettings.cpp
@@ -3,6 +3,7 @@
 #include "Replication/Editor/Model/StreamExtenderBySettings.h"
 
 #include "Replication/Data/ConcertPropertySelection.h"
+#include "Replication/Editor/Model/Extension/IStreamExtensionContext.h"
 #include "Replication/Settings/ConcertReplicationEditorSettings.h"
 
 namespace UE::ConcertClientSharedSlate
@@ -11,31 +12,20 @@ namespace UE::ConcertClientSharedSlate
 		: ReplicationSettingsAttribute(MoveTemp(InReplicationSettingsAttribute))
 	{}
 
-	void FStreamExtenderBySettings::ExtendObjectProperties(UObject& Object, FEnumerateProperties ForEachPropertyToAdd)
+	void FStreamExtenderBySettings::ExtendStream(UObject& ExtendedObject, IStreamExtensionContext& Context)
 	{
 		const FConcertReplicationEditorSettings* Settings = ReplicationSettingsAttribute.Get();
 		if (!ensure(Settings))
 		{
 			return;
 		}
-		
-		Settings->AddDefaultPropertiesFromSettings(*Object.GetClass(), [&ForEachPropertyToAdd](FConcertPropertyChain&& PropertyChain)
+		Settings->AddDefaultPropertiesFromSettings(*ExtendedObject.GetClass(), [&ExtendedObject, &Context](FConcertPropertyChain&& PropertyChain)
 		{
-			ForEachPropertyToAdd(MoveTemp(PropertyChain));
+			Context.AddPropertyTo(ExtendedObject, MoveTemp(PropertyChain));
 		});
-	}
-
-	void FStreamExtenderBySettings::AppendAdditionalObjects(UObject& Object, FEnumerateObjects ForEachAdditionalObject)
-	{
-		const FConcertReplicationEditorSettings* Settings = ReplicationSettingsAttribute.Get();
-		if (!ensure(Settings))
-		{
-			return;
-		}
-		
-		Settings->AddAdditionalObjectsFromSettings(Object, [&ForEachAdditionalObject](UObject& AdditionalObject)
+		Settings->AddAdditionalObjectsFromSettings(ExtendedObject, [&ExtendedObject, &Context](UObject& AdditionalObject)
 		{
-			ForEachAdditionalObject(AdditionalObject);
+			Context.AddAdditionalObject(AdditionalObject);
 		});
 	}
 }
diff --git a/Engine/Plugins/Developer/Concert/ConcertUI/ConcertClientSharedSlate/Source/ConcertClientSharedSlate/Public/Replication/Editor/Model/Extension/IStreamExtender.h b/Engine/Plugins/Developer/Concert/ConcertUI/ConcertClientSharedSlate/Source/ConcertClientSharedSlate/Public/Replication/Editor/Model/Extension/IStreamExtender.h
new file mode 100644
--- /dev/null
+++ b/Engine/Plugins/Developer/Concert/ConcertUI/ConcertClientSharedSlate/Source/ConcertClientSharedSlate/Public/Replication/Editor/Model/Extension/IStreamExtender.h
@@ -0,0 +1,30 @@
+// Copyright Epic Games, Inc. All Rights Reserved.
+
+#pragma once
+
+#include "Templates/Function.h"
+
+class UObject;
+struct FConcertPropertyChain;
+
+namespace UE::ConcertClientSharedSlate
+{
+	class IStreamExtensionContext;
+
+	/** When an object is added to IEditableReplicationStreamModel, these functions are called to allow adding additional objects and properties. */
+	class CONCERTCLIENTSHAREDSLATE_API IStreamExtender
+	{
+	public:
+
+		/**
+		 * Allows extending the stream when objects are added.
+		 * This function is called recursively for any additional objects you add.
+		 * 
+		 * @param ExtendedObject Object that was added
+		 * @param Context Use for adding properties to ExtendedObject or to additional objects.
+		 */
+		virtual void ExtendStream(UObject& ExtendedObject, IStreamExtensionContext& Context) = 0;
+
+		virtual ~IStreamExtender() = default;
+	};
+}
diff --git a/Engine/Plugins/Developer/Concert/ConcertUI/ConcertClientSharedSlate/Source/ConcertClientSharedSlate/Public/Replication/Editor/Model/Extension/IStreamExtensionContext.h b/Engine/Plugins/Developer/Concert/ConcertUI/ConcertClientSharedSlate/Source/ConcertClientSharedSlate/Public/Replication/Editor/Model/Extension/IStreamExtensionContext.h
new file mode 100644
--- /dev/null
+++ b/Engine/Plugins/Developer/Concert/ConcertUI/ConcertClientSharedSlate/Source/ConcertClientSharedSlate/Public/Replication/Editor/Model/Extension/IStreamExtensionContext.h
@@ -0,0 +1,37 @@
+// Copyright Epic Games, Inc. All Rights Reserved.
+
+#pragma once
+
+#include "Templates/Function.h"
+
+class UObject;
+struct FConcertPropertyChain;
+
+namespace UE::ConcertClientSharedSlate
+{
+	/** Allows IStreamExtender to add default replication settings like properties and adding additional objects. */
+	class CONCERTCLIENTSHAREDSLATE_API IStreamExtensionContext
+	{
+	public:
+
+		/**
+		 * Extends the stream with the given object, if not yet added, and adds the property to it.
+		 * If Object was not yet in the stream, IStreamExtender::ExtendStream will be called for Object.
+		 * 
+		 * @param Object Object to add, if not yet present, and assign a property to
+		 * @param PropertyChain Property chain to associate with the object
+		 */
+		virtual void AddPropertyTo(UObject& Object, FConcertPropertyChain&& PropertyChain) = 0;
+
+		/**
+		 * Adds an additional object.
+		 * 
+		 * If Object was not yet in the stream, IStreamExtender::ExtendStream will be called for Object.
+		 * 
+		 * @param Object The additional object to add
+		 */
+		virtual void AddAdditionalObject(UObject& Object) = 0;
+
+		virtual ~IStreamExtensionContext() = default;
+	};
+}
diff --git a/Engine/Plugins/Developer/Concert/ConcertUI/ConcertClientSharedSlate/Source/ConcertClientSharedSlate/Public/Replication/Editor/Model/IStreamExtender.h b/Engine/Plugins/Developer/Concert/ConcertUI/ConcertClientSharedSlate/Source/ConcertClientSharedSlate/Public/Replication/Editor/Model/IStreamExtender.h
deleted file mode 100644
diff --git a/Engine/Plugins/Developer/Concert/ConcertUI/ConcertClientSharedSlate/Source/ConcertClientSharedSlate/Public/Replication/Editor/Model/StreamExtenderBySettings.h b/Engine/Plugins/Developer/Concert/ConcertUI/ConcertClientSharedSlate/Source/ConcertClientSharedSlate/Public/Replication/Editor/Model/StreamExtenderBySettings.h
--- a/Engine/Plugins/Developer/Concert/ConcertUI/ConcertClientSharedSlate/Source/ConcertClientSharedSlate/Public/Replication/Editor/Model/StreamExtenderBySettings.h
+++ b/Engine/Plugins/Developer/Concert/ConcertUI/ConcertClientSharedSlate/Source/ConcertClientSharedSlate/Public/Replication/Editor/Model/StreamExtenderBySettings.h
@@ -2,7 +2,7 @@
 
 #pragma once
 
-#include "IStreamExtender.h"
+#include "Extension/IStreamExtender.h"
 #include "Misc/Attribute.h"
 
 class UObject;
@@ -18,8 +18,7 @@ namespace UE::ConcertClientSharedSlate
 		FStreamExtenderBySettings(TAttribute<const FConcertReplicationEditorSettings*> InReplicationSettingsAttribute);
 
 		//~ Begin IStreamExtender Interface
-		virtual void ExtendObjectProperties(UObject& Object, FEnumerateProperties ForEachPropertyToAdd) override;
-		virtual void AppendAdditionalObjects(UObject& Object, FEnumerateObjects ForEachAdditionalObject) override;
+		virtual void ExtendStream(UObject& ExtendedObject, IStreamExtensionContext& Context) override;
 		//~ End IStreamExtender Interface
 
 	private:
