Fix subobject display names not showing for Blueprints

diff --git a/Engine/Plugins/Developer/Concert/ConcertUI/ConcertClientSharedSlate/Source/ConcertClientSharedSlate/Private/Replication/Editor/Model/Subobject/ComponentHierarchySubobjectModel.cpp b/Engine/Plugins/Developer/Concert/ConcertUI/ConcertClientSharedSlate/Source/ConcertClientSharedSlate/Private/Replication/Editor/Model/Subobject/ComponentHierarchySubobjectModel.cpp
--- a/Engine/Plugins/Developer/Concert/ConcertUI/ConcertClientSharedSlate/Source/ConcertClientSharedSlate/Private/Replication/Editor/Model/Subobject/ComponentHierarchySubobjectModel.cpp
+++ b/Engine/Plugins/Developer/Concert/ConcertUI/ConcertClientSharedSlate/Source/ConcertClientSharedSlate/Private/Replication/Editor/Model/Subobject/ComponentHierarchySubobjectModel.cpp
@@ -38,8 +38,11 @@ namespace UE::ConcertClientSharedSlate
 				for (const FSubobjectDataHandle& Handle : Handles)
 				{
 					const FSubobjectData* SubobjectData = Handle.GetData();
-					const UObject* Object = SubobjectData->GetObject();
-					ObjectMetaData.Add(Object, { FText::FromString(SubobjectData->GetDisplayString()) });
+					const UObject* Object = SubobjectData->FindComponentInstanceInActor(Actor);
+					if (Object)
+					{
+						ObjectMetaData.Add(Object, { FText::FromString(SubobjectData->GetDisplayString()) });
+					}
 				}
 			}
 			OnHierarchyChangedDelegate.Broadcast();
@@ -116,4 +119,4 @@ namespace UE::ConcertClientSharedSlate
 	{
 		return Cast<AActor>(TopLevelObject.ResolveObject());
 	}
-}
+}
\ No newline at end of file
diff --git a/Engine/Plugins/Developer/Concert/ConcertUI/ConcertClientSharedSlate/Source/ConcertClientSharedSlate/Private/Replication/Editor/View/Tree/SelectionViewerColumns.cpp b/Engine/Plugins/Developer/Concert/ConcertUI/ConcertClientSharedSlate/Source/ConcertClientSharedSlate/Private/Replication/Editor/View/Tree/SelectionViewerColumns.cpp
--- a/Engine/Plugins/Developer/Concert/ConcertUI/ConcertClientSharedSlate/Source/ConcertClientSharedSlate/Private/Replication/Editor/View/Tree/SelectionViewerColumns.cpp
+++ b/Engine/Plugins/Developer/Concert/ConcertUI/ConcertClientSharedSlate/Source/ConcertClientSharedSlate/Private/Replication/Editor/View/Tree/SelectionViewerColumns.cpp
@@ -53,6 +53,8 @@ namespace UE::ConcertClientSharedSlate::ReplicationColumns::TopLevel
 					const FText Text = GetDisplayText(Args.RowData);
 					
 					return SNew(SHorizontalBox)
+						.ToolTipText(FText::FromString(Args.RowData.GetObjectPath().ToString()))
+					
 						+SHorizontalBox::Slot()
 						.AutoWidth()
 						.HAlign(HAlign_Center)
