Automatically expand newly added objects

diff --git a/Engine/Plugins/Developer/Concert/ConcertUI/ConcertClientSharedSlate/Source/ConcertClientSharedSlate/Private/Replication/Editor/View/ObjectViewer/SReplicationStreamViewer.cpp b/Engine/Plugins/Developer/Concert/ConcertUI/ConcertClientSharedSlate/Source/ConcertClientSharedSlate/Private/Replication/Editor/View/ObjectViewer/SReplicationStreamViewer.cpp
--- a/Engine/Plugins/Developer/Concert/ConcertUI/ConcertClientSharedSlate/Source/ConcertClientSharedSlate/Private/Replication/Editor/View/ObjectViewer/SReplicationStreamViewer.cpp
+++ b/Engine/Plugins/Developer/Concert/ConcertUI/ConcertClientSharedSlate/Source/ConcertClientSharedSlate/Private/Replication/Editor/View/ObjectViewer/SReplicationStreamViewer.cpp
@@ -140,9 +140,9 @@ namespace UE::ConcertClientSharedSlate
 			}
 		}
 
-		if (ItemsToExpand.IsEmpty())
+		if (!ItemsToExpand.IsEmpty())
 		{
-			ReplicatedObjects->SetSelectedItems(ItemsToExpand, true);
+			ReplicatedObjects->SetExpandedItems(ItemsToExpand, true);
 		}
 	}
 	
