Refactor util concert replication settings types to be more generic, in preparation to reuse for replication frequency settings

diff --git a/Engine/Plugins/Developer/Concert/ConcertApp/MultiUserClient/Config/BaseMultiUserClient.ini b/Engine/Plugins/Developer/Concert/ConcertApp/MultiUserClient/Config/BaseMultiUserClient.ini
--- a/Engine/Plugins/Developer/Concert/ConcertApp/MultiUserClient/Config/BaseMultiUserClient.ini
+++ b/Engine/Plugins/Developer/Concert/ConcertApp/MultiUserClient/Config/BaseMultiUserClient.ini
@@ -1,2 +1,2 @@
 [/Script/MultiUserClient.MultiUserReplicationSettings]
-ReplicationEditorSettings=(DefaultPropertySelection=(("/Script/Engine.SceneComponent", (DefaultSelectedProperties=("RelativeLocation","RelativeLocation.X","RelativeLocation.Y","RelativeLocation.Z","RelativeRotation"))),("/Script/CinematicCamera.CineCameraComponent", ()),("/Script/Engine.CameraComponent", ())),DefaultSubobjectSelection=(("/Script/Engine.Actor", (IncludeClasses=("/Script/Engine.StaticMeshComponent","/Script/Engine.CameraComponent","/Script/CinematicCamera.CineCameraComponent")))))
+ReplicationEditorSettings=(DefaultPropertySelection=(("/Script/Engine.SceneComponent", (DefaultSelectedProperties=("RelativeLocation","RelativeLocation.X","RelativeLocation.Y","RelativeLocation.Z","RelativeRotation"))),("/Script/CinematicCamera.CineCameraComponent", ()),("/Script/Engine.CameraComponent", ())),DefaultAddedSubobjectRules=(SubobjectMatchingRules=(("/Script/Engine.Actor", (IncludeClasses=("/Script/Engine.StaticMeshComponent","/Script/Engine.CameraComponent","/Script/CinematicCamera.CineCameraComponent"))))))
\ No newline at end of file
diff --git a/Engine/Plugins/Developer/Concert/ConcertApp/MultiUserClient/Source/MultiUserClient/Private/MultiUserReplicationSettings.h b/Engine/Plugins/Developer/Concert/ConcertApp/MultiUserClient/Source/MultiUserClient/Private/MultiUserReplicationSettings.h
--- a/Engine/Plugins/Developer/Concert/ConcertApp/MultiUserClient/Source/MultiUserClient/Private/MultiUserReplicationSettings.h
+++ b/Engine/Plugins/Developer/Concert/ConcertApp/MultiUserClient/Source/MultiUserClient/Private/MultiUserReplicationSettings.h
@@ -3,7 +3,7 @@
 #pragma once
 
 #include "UObject/Object.h"
-#include "Replication/Settings/ConcertReplicationEditorSettings.h"
+#include "Replication/Settings/ConcertStreamObjectAutoBindingRules.h"
 #include "MultiUserReplicationSettings.generated.h"
 
 UCLASS(config=MultiUserClient)
@@ -14,7 +14,10 @@ public:
 
 	static UMultiUserReplicationSettings* Get() { return GetMutableDefault<UMultiUserReplicationSettings>(); }
 
-	/** Settings that affect the replication editor */
+	/**
+	 * When you add an object via the stream editor, you may want to automatically bind properties and add additional subobjects.
+	 * Here you can specify the rules to achieve this.
+	 */
 	UPROPERTY(Config, Category = "Replication Settings", EditAnywhere)
-	FConcertReplicationEditorSettings ReplicationEditorSettings;
+	FConcertStreamObjectAutoBindingRules ReplicationEditorSettings;
 };
\ No newline at end of file
diff --git a/Engine/Plugins/Developer/Concert/ConcertApp/MultiUserClient/Source/MultiUserClient/Private/Replication/Stream/Discovery/MultiUserStreamExtender.cpp b/Engine/Plugins/Developer/Concert/ConcertApp/MultiUserClient/Source/MultiUserClient/Private/Replication/Stream/Discovery/MultiUserStreamExtender.cpp
--- a/Engine/Plugins/Developer/Concert/ConcertApp/MultiUserClient/Source/MultiUserClient/Private/Replication/Stream/Discovery/MultiUserStreamExtender.cpp
+++ b/Engine/Plugins/Developer/Concert/ConcertApp/MultiUserClient/Source/MultiUserClient/Private/Replication/Stream/Discovery/MultiUserStreamExtender.cpp
@@ -14,7 +14,7 @@ namespace UE::MultiUserClient
 		: ClientId(InClientId)
 		, ExtendBySettings(
 			// Use MU settings for auto adding properties & objects
-			TAttribute<const FConcertReplicationEditorSettings*>::CreateLambda([]()
+			TAttribute<const FConcertStreamObjectAutoBindingRules*>::CreateLambda([]()
 			{
 				return &UMultiUserReplicationSettings::Get()->ReplicationEditorSettings;
 			}))
diff --git a/Engine/Plugins/Developer/Concert/ConcertUI/ConcertClientSharedSlate/Source/ConcertClientSharedSlate/Private/Replication/Editor/Model/GenericReplicationStreamModel.h b/Engine/Plugins/Developer/Concert/ConcertUI/ConcertClientSharedSlate/Source/ConcertClientSharedSlate/Private/Replication/Editor/Model/GenericReplicationStreamModel.h
--- a/Engine/Plugins/Developer/Concert/ConcertUI/ConcertClientSharedSlate/Source/ConcertClientSharedSlate/Private/Replication/Editor/Model/GenericReplicationStreamModel.h
+++ b/Engine/Plugins/Developer/Concert/ConcertUI/ConcertClientSharedSlate/Source/ConcertClientSharedSlate/Private/Replication/Editor/Model/GenericReplicationStreamModel.h
@@ -6,7 +6,7 @@
 #include "Misc/Attribute.h"
 #include "UObject/WeakObjectPtrTemplates.h"
 
-struct FConcertReplicationEditorSettings;
+struct FConcertStreamObjectAutoBindingRules;
 struct FObjectReplicationMap;
 
 namespace UE::ConcertClientSharedSlate
diff --git a/Engine/Plugins/Developer/Concert/ConcertUI/ConcertClientSharedSlate/Source/ConcertClientSharedSlate/Private/Replication/Editor/Model/StreamExtenderBySettings.cpp b/Engine/Plugins/Developer/Concert/ConcertUI/ConcertClientSharedSlate/Source/ConcertClientSharedSlate/Private/Replication/Editor/Model/StreamExtenderBySettings.cpp
--- a/Engine/Plugins/Developer/Concert/ConcertUI/ConcertClientSharedSlate/Source/ConcertClientSharedSlate/Private/Replication/Editor/Model/StreamExtenderBySettings.cpp
+++ b/Engine/Plugins/Developer/Concert/ConcertUI/ConcertClientSharedSlate/Source/ConcertClientSharedSlate/Private/Replication/Editor/Model/StreamExtenderBySettings.cpp
@@ -4,17 +4,17 @@
 
 #include "Replication/Data/ConcertPropertySelection.h"
 #include "Replication/Editor/Model/Extension/IStreamExtensionContext.h"
-#include "Replication/Settings/ConcertReplicationEditorSettings.h"
+#include "Replication/Settings/ConcertStreamObjectAutoBindingRules.h"
 
 namespace UE::ConcertClientSharedSlate
 {
-	FStreamExtenderBySettings::FStreamExtenderBySettings(TAttribute<const FConcertReplicationEditorSettings*> InReplicationSettingsAttribute)
+	FStreamExtenderBySettings::FStreamExtenderBySettings(TAttribute<const FConcertStreamObjectAutoBindingRules*> InReplicationSettingsAttribute)
 		: ReplicationSettingsAttribute(MoveTemp(InReplicationSettingsAttribute))
 	{}
 
 	void FStreamExtenderBySettings::ExtendStream(UObject& ExtendedObject, IStreamExtensionContext& Context)
 	{
-		const FConcertReplicationEditorSettings* Settings = ReplicationSettingsAttribute.Get();
+		const FConcertStreamObjectAutoBindingRules* Settings = ReplicationSettingsAttribute.Get();
 		if (!ensure(Settings))
 		{
 			return;
@@ -23,7 +23,7 @@ namespace UE::ConcertClientSharedSlate
 		{
 			Context.AddPropertyTo(ExtendedObject, MoveTemp(PropertyChain));
 		});
-		Settings->AddAdditionalObjectsFromSettings(ExtendedObject, [&ExtendedObject, &Context](UObject& AdditionalObject)
+		Settings->DefaultAddedSubobjectRules.MatchSubobjectsRecursivelyFor(ExtendedObject, [&ExtendedObject, &Context](UObject& AdditionalObject)
 		{
 			Context.AddAdditionalObject(AdditionalObject);
 		});
diff --git a/Engine/Plugins/Developer/Concert/ConcertUI/ConcertClientSharedSlate/Source/ConcertClientSharedSlate/Private/Replication/Editor/View/ObjectEditor/SBaseReplicationStreamEditor.cpp b/Engine/Plugins/Developer/Concert/ConcertUI/ConcertClientSharedSlate/Source/ConcertClientSharedSlate/Private/Replication/Editor/View/ObjectEditor/SBaseReplicationStreamEditor.cpp
--- a/Engine/Plugins/Developer/Concert/ConcertUI/ConcertClientSharedSlate/Source/ConcertClientSharedSlate/Private/Replication/Editor/View/ObjectEditor/SBaseReplicationStreamEditor.cpp
+++ b/Engine/Plugins/Developer/Concert/ConcertUI/ConcertClientSharedSlate/Source/ConcertClientSharedSlate/Private/Replication/Editor/View/ObjectEditor/SBaseReplicationStreamEditor.cpp
@@ -10,7 +10,7 @@
 #include "Replication/Editor/Model/ReplicatedObjectData.h"
 #include "Replication/Editor/View/DisplayUtils.h"
 #include "Replication/Editor/View/ObjectViewer/SReplicationStreamViewer.h"
-#include "Replication/Settings/ConcertReplicationEditorSettings.h"
+#include "Replication/Settings/ConcertStreamObjectAutoBindingRules.h"
 
 #include "Algo/AnyOf.h"
 #include "Replication/Editor/Model/ISubobjectModel.h"
diff --git a/Engine/Plugins/Developer/Concert/ConcertUI/ConcertClientSharedSlate/Source/ConcertClientSharedSlate/Private/Replication/Editor/View/ObjectEditor/SBaseReplicationStreamEditor.h b/Engine/Plugins/Developer/Concert/ConcertUI/ConcertClientSharedSlate/Source/ConcertClientSharedSlate/Private/Replication/Editor/View/ObjectEditor/SBaseReplicationStreamEditor.h
--- a/Engine/Plugins/Developer/Concert/ConcertUI/ConcertClientSharedSlate/Source/ConcertClientSharedSlate/Private/Replication/Editor/View/ObjectEditor/SBaseReplicationStreamEditor.h
+++ b/Engine/Plugins/Developer/Concert/ConcertUI/ConcertClientSharedSlate/Source/ConcertClientSharedSlate/Private/Replication/Editor/View/ObjectEditor/SBaseReplicationStreamEditor.h
@@ -12,7 +12,7 @@
 
 class SHorizontalBox;
 struct FConcertPropertyChain;
-struct FConcertReplicationEditorSettings;
+struct FConcertStreamObjectAutoBindingRules;
 
 namespace UE::ConcertClientSharedSlate
 {
diff --git a/Engine/Plugins/Developer/Concert/ConcertUI/ConcertClientSharedSlate/Source/ConcertClientSharedSlate/Private/Replication/Editor/View/ObjectEditor/SDefaultReplicationStreamEditor.h b/Engine/Plugins/Developer/Concert/ConcertUI/ConcertClientSharedSlate/Source/ConcertClientSharedSlate/Private/Replication/Editor/View/ObjectEditor/SDefaultReplicationStreamEditor.h
--- a/Engine/Plugins/Developer/Concert/ConcertUI/ConcertClientSharedSlate/Source/ConcertClientSharedSlate/Private/Replication/Editor/View/ObjectEditor/SDefaultReplicationStreamEditor.h
+++ b/Engine/Plugins/Developer/Concert/ConcertUI/ConcertClientSharedSlate/Source/ConcertClientSharedSlate/Private/Replication/Editor/View/ObjectEditor/SDefaultReplicationStreamEditor.h
@@ -12,7 +12,7 @@
 
 class SHorizontalBox;
 struct FConcertPropertyChain;
-struct FConcertReplicationEditorSettings;
+struct FConcertStreamObjectAutoBindingRules;
 
 namespace UE::ConcertClientSharedSlate
 {
diff --git a/Engine/Plugins/Developer/Concert/ConcertUI/ConcertClientSharedSlate/Source/ConcertClientSharedSlate/Private/Replication/Settings/ConcertPerClassSubobjectMatchingRules.cpp b/Engine/Plugins/Developer/Concert/ConcertUI/ConcertClientSharedSlate/Source/ConcertClientSharedSlate/Private/Replication/Settings/ConcertPerClassSubobjectMatchingRules.cpp
new file mode 100644
--- /dev/null
+++ b/Engine/Plugins/Developer/Concert/ConcertUI/ConcertClientSharedSlate/Source/ConcertClientSharedSlate/Private/Replication/Settings/ConcertPerClassSubobjectMatchingRules.cpp
@@ -0,0 +1,37 @@
+// Copyright Epic Games, Inc. All Rights Reserved.
+
+#include "Replication/Settings/ConcertPerClassSubobjectMatchingRules.h"
+
+#include "UObject/Class.h"
+
+namespace UE::ConcertClientSharedSlate::DefaultSubobjects
+{
+	static void InternalAddAdditionalObjectsFromSettings(UClass& StartClass, const FConcertPerClassSubobjectMatchingRules& Settings, const UObject& AddedObject, TFunctionRef<void(UObject&)> OnSubobjectMatched)
+	{
+		// Find the most specialized class properties
+		UClass* Current = &StartClass;
+		const FConcertInheritableSubobjectMatchingRules* DefaultSubobjects = nullptr;
+		for (; Current && !DefaultSubobjects; Current = Current->GetSuperClass())
+		{
+			DefaultSubobjects = Settings.SubobjectMatchingRules.Find(Current);
+			if (!DefaultSubobjects)
+			{
+				continue;
+			}
+			
+			DefaultSubobjects->MatchToSubobjectsIn(AddedObject, OnSubobjectMatched);
+			// Recurse super structs
+			if (UClass* Parent = Current->GetSuperClass()
+				; Parent && DefaultSubobjects->bInheritFromBase)
+			{
+				InternalAddAdditionalObjectsFromSettings(*Parent, Settings, AddedObject, OnSubobjectMatched);
+			}
+		}
+	}
+}
+
+void FConcertPerClassSubobjectMatchingRules::MatchSubobjectsRecursivelyFor(const UObject& Object, TFunctionRef<void(UObject&)> OnSubobjectMatched) const
+{
+	checkf(!Object.IsA<UClass>(), TEXT("Pass in the UObject instanced directly, not its class!"));
+	UE::ConcertClientSharedSlate::DefaultSubobjects::InternalAddAdditionalObjectsFromSettings(*Object.GetClass(), *this, Object, OnSubobjectMatched);
+}
\ No newline at end of file
diff --git a/Engine/Plugins/Developer/Concert/ConcertUI/ConcertClientSharedSlate/Source/ConcertClientSharedSlate/Private/Replication/Settings/ConcertReplicationEditorSettings.cpp b/Engine/Plugins/Developer/Concert/ConcertUI/ConcertClientSharedSlate/Source/ConcertClientSharedSlate/Private/Replication/Settings/ConcertReplicationEditorSettings.cpp
deleted file mode 100644
diff --git a/Engine/Plugins/Developer/Concert/ConcertUI/ConcertClientSharedSlate/Source/ConcertClientSharedSlate/Private/Replication/Settings/ConcertStreamObjectAutoBindingRules.cpp b/Engine/Plugins/Developer/Concert/ConcertUI/ConcertClientSharedSlate/Source/ConcertClientSharedSlate/Private/Replication/Settings/ConcertStreamObjectAutoBindingRules.cpp
new file mode 100644
--- /dev/null
+++ b/Engine/Plugins/Developer/Concert/ConcertUI/ConcertClientSharedSlate/Source/ConcertClientSharedSlate/Private/Replication/Settings/ConcertStreamObjectAutoBindingRules.cpp
@@ -0,0 +1,74 @@
+// Copyright Epic Games, Inc. All Rights Reserved.
+
+#include "Replication/Settings/ConcertStreamObjectAutoBindingRules.h"
+
+#include "ConcertLogGlobal.h"
+#include "Replication/PropertyChainUtils.h"
+
+#include "Algo/IndexOf.h"
+#include "UObject/Class.h"
+
+namespace UE::ConcertClientSharedSlate::DefaultProperties
+{
+	static void ApplyDefaultPropertySelection(TFunctionRef<void(FConcertPropertyChain&& Chain)> Callback, const FConcertDefaultPropertySelection& Selection, UStruct& Class)
+	{
+		// Preparse the the FStrings into paths
+		TArray<TArray<FName>> Paths;
+		Paths.Reserve(Selection.DefaultSelectedProperties.Num());
+		for (const FString& StringPath : Selection.DefaultSelectedProperties)
+		{
+			TArray<FString> Parts;
+			StringPath.ParseIntoArray(Parts, TEXT("."));
+			TArray<FName> NameArray;
+			Algo::Transform(Parts, NameArray, [](const FString& String){ return *String; });
+			Paths.Emplace(MoveTemp(NameArray));
+		}
+
+		// This walks through the entire property hierarchy once
+		using namespace ConcertSyncCore::PropertyChain;
+		BulkConstructConcertChainsFromPaths(Class, Paths.Num(), [&Callback, &Paths](const FArchiveSerializedPropertyChain& Chain, const FProperty& LeafProperty)
+		{
+			const int32 IndexOfMatches = Algo::IndexOfByPredicate(Paths, [&Chain, &LeafProperty](const TArray<FName>& Path){ return DoPathAndChainsMatch(Path, Chain, LeafProperty); });
+			const bool bMatches = IndexOfMatches != INDEX_NONE;
+			if (bMatches)
+			{
+				Callback(FConcertPropertyChain(&Chain, LeafProperty));
+				Paths.RemoveAtSwap(IndexOfMatches);
+			}
+			return true;
+		});
+
+		// Warn the user of unmatched properties so they can de-clutter their settings of useless entries
+		for (const TArray<FName>& Path : Paths)
+		{
+			UE_LOG(LogConcert, Warning,
+				TEXT("Path %s was unmatched. Fix your settings for class %s."),
+				*FString::JoinBy(Path, TEXT("."), [](const FName& Name){ return Name.ToString(); }),
+				*Class.GetName()
+				);
+		}
+	}
+}
+
+void FConcertStreamObjectAutoBindingRules::AddDefaultPropertiesFromSettings(UClass& Class, TFunctionRef<void(FConcertPropertyChain&& Chain)> Callback) const
+{
+	// Find the most specialized class properties
+	UClass* Current = &Class;
+	const FConcertDefaultPropertySelection* DefaultProperties = nullptr;
+	for (; Current && !DefaultProperties; Current = Current->GetSuperClass())
+	{
+		DefaultProperties = DefaultPropertySelection.Find(Current);
+		if (!DefaultProperties)
+		{
+			continue;
+		}
+			
+		UE::ConcertClientSharedSlate::DefaultProperties::ApplyDefaultPropertySelection(Callback, *DefaultProperties, *Current);
+		// Recurse super structs
+		if (UClass* Parent = Current->GetSuperClass()
+			; Parent && DefaultProperties->bInheritFromBase)
+		{
+			AddDefaultPropertiesFromSettings(*Parent, Callback);
+		}
+	}
+}
\ No newline at end of file
diff --git a/Engine/Plugins/Developer/Concert/ConcertUI/ConcertClientSharedSlate/Source/ConcertClientSharedSlate/Private/Replication/Settings/ConcertSubobjectMatchingRules.cpp b/Engine/Plugins/Developer/Concert/ConcertUI/ConcertClientSharedSlate/Source/ConcertClientSharedSlate/Private/Replication/Settings/ConcertSubobjectMatchingRules.cpp
new file mode 100644
--- /dev/null
+++ b/Engine/Plugins/Developer/Concert/ConcertUI/ConcertClientSharedSlate/Source/ConcertClientSharedSlate/Private/Replication/Settings/ConcertSubobjectMatchingRules.cpp
@@ -0,0 +1,67 @@
+// Copyright Epic Games, Inc. All Rights Reserved.
+
+#include "Replication/Settings/ConcertSubobjectMatchingRules.h"
+
+#include "Components/ActorComponent.h"
+#include "Internationalization/Regex.h"
+
+namespace UE::ConcertClientSharedSlate
+{
+	static bool MatchesAnyRegex(const FString& Input, const TSet<FString>& AllRegex)
+	{
+		for (const FString& RegexString : AllRegex)
+		{
+			const FRegexPattern Pattern(RegexString);
+			FRegexMatcher Matcher(Pattern, Input);
+			if (Matcher.FindNext())
+			{
+				return true;
+			}
+		}
+
+		return false;
+	}
+}
+
+void FConcertSubobjectMatchingRules::MatchToSubobjectsIn(const UObject& AddedObject, TFunctionRef<void(UObject&)> OnSubobjectMatched) const
+{
+	// Do not search recursively so FurtherObjectsCallback can decide to call AddAdditionalObjectsFromSettings again on the newly added objects.
+	constexpr bool bIncludeNested = false;
+	ForEachObjectWithOuter(&AddedObject, [this, &OnSubobjectMatched](UObject* Subobject)
+	{
+		// Has exclusion regex?
+		const bool bShouldExclude = UE::ConcertClientSharedSlate::MatchesAnyRegex(Subobject->GetName(), ExcludeSubobjectRegex);
+		if (bShouldExclude)
+		{
+			return;
+		}
+
+		// Was told to add all UActorComponents?
+		const bool bIncludeAllSubobjects = IncludeAllOption == EConcertIncludeAllSubobjectsType::AllSubobjects;
+		const bool bIncludeDueToComponent = Subobject->IsA(UActorComponent::StaticClass())
+			&& IncludeAllOption == EConcertIncludeAllSubobjectsType::AllComponents;
+		if (bIncludeAllSubobjects || bIncludeDueToComponent)
+		{
+			OnSubobjectMatched(*Subobject);
+			return;
+		}
+
+		// Has configured class?
+		UClass* CurrentClass = Subobject->GetClass();
+		for (; CurrentClass; CurrentClass = CurrentClass->GetSuperClass())
+		{
+			if (IncludeClasses.Contains(Subobject->GetClass()))
+			{
+				OnSubobjectMatched(*Subobject);
+				return;
+			}
+		}
+
+		// Has inclusion regex?
+		const bool bShouldIncludeByRegex = UE::ConcertClientSharedSlate::MatchesAnyRegex(Subobject->GetName(), IncludeSubobjectRegex);
+		if (bShouldIncludeByRegex)
+		{
+			OnSubobjectMatched(*Subobject);
+		}
+	}, bIncludeNested);
+}
\ No newline at end of file
diff --git a/Engine/Plugins/Developer/Concert/ConcertUI/ConcertClientSharedSlate/Source/ConcertClientSharedSlate/Public/Replication/Editor/Model/StreamExtenderBySettings.h b/Engine/Plugins/Developer/Concert/ConcertUI/ConcertClientSharedSlate/Source/ConcertClientSharedSlate/Public/Replication/Editor/Model/StreamExtenderBySettings.h
--- a/Engine/Plugins/Developer/Concert/ConcertUI/ConcertClientSharedSlate/Source/ConcertClientSharedSlate/Public/Replication/Editor/Model/StreamExtenderBySettings.h
+++ b/Engine/Plugins/Developer/Concert/ConcertUI/ConcertClientSharedSlate/Source/ConcertClientSharedSlate/Public/Replication/Editor/Model/StreamExtenderBySettings.h
@@ -6,7 +6,7 @@
 #include "Misc/Attribute.h"
 
 class UObject;
-struct FConcertReplicationEditorSettings;
+struct FConcertStreamObjectAutoBindingRules;
 
 namespace UE::ConcertClientSharedSlate
 {
@@ -15,7 +15,7 @@ namespace UE::ConcertClientSharedSlate
 	{
 	public:
 		
-		FStreamExtenderBySettings(TAttribute<const FConcertReplicationEditorSettings*> InReplicationSettingsAttribute);
+		FStreamExtenderBySettings(TAttribute<const FConcertStreamObjectAutoBindingRules*> InReplicationSettingsAttribute);
 
 		//~ Begin IStreamExtender Interface
 		virtual void ExtendStream(UObject& ExtendedObject, IStreamExtensionContext& Context) override;
@@ -24,6 +24,6 @@ namespace UE::ConcertClientSharedSlate
 	private:
 		
 		/** Optional settings for auto adding common properties and objects. */
-		TAttribute<const FConcertReplicationEditorSettings*> ReplicationSettingsAttribute;
+		TAttribute<const FConcertStreamObjectAutoBindingRules*> ReplicationSettingsAttribute;
 	};
 }
diff --git a/Engine/Plugins/Developer/Concert/ConcertUI/ConcertClientSharedSlate/Source/ConcertClientSharedSlate/Public/Replication/ReplicationWidgetFactories.h b/Engine/Plugins/Developer/Concert/ConcertUI/ConcertClientSharedSlate/Source/ConcertClientSharedSlate/Public/Replication/ReplicationWidgetFactories.h
--- a/Engine/Plugins/Developer/Concert/ConcertUI/ConcertClientSharedSlate/Source/ConcertClientSharedSlate/Public/Replication/ReplicationWidgetFactories.h
+++ b/Engine/Plugins/Developer/Concert/ConcertUI/ConcertClientSharedSlate/Source/ConcertClientSharedSlate/Public/Replication/ReplicationWidgetFactories.h
@@ -13,7 +13,7 @@
 class UObject;
 class SWidget;
 
-struct FConcertReplicationEditorSettings;
+struct FConcertStreamObjectAutoBindingRules;
 struct FObjectReplicationMap;
 
 namespace UE::ConcertClientSharedSlate
diff --git a/Engine/Plugins/Developer/Concert/ConcertUI/ConcertClientSharedSlate/Source/ConcertClientSharedSlate/Public/Replication/Settings/ConcertDefaultPropertySelection.h b/Engine/Plugins/Developer/Concert/ConcertUI/ConcertClientSharedSlate/Source/ConcertClientSharedSlate/Public/Replication/Settings/ConcertDefaultPropertySelection.h
deleted file mode 100644
diff --git a/Engine/Plugins/Developer/Concert/ConcertUI/ConcertClientSharedSlate/Source/ConcertClientSharedSlate/Public/Replication/Settings/ConcertDefaultSubobjectSelection.h b/Engine/Plugins/Developer/Concert/ConcertUI/ConcertClientSharedSlate/Source/ConcertClientSharedSlate/Public/Replication/Settings/ConcertDefaultSubobjectSelection.h
deleted file mode 100644
diff --git a/Engine/Plugins/Developer/Concert/ConcertUI/ConcertClientSharedSlate/Source/ConcertClientSharedSlate/Public/Replication/Settings/ConcertPerClassSubobjectMatchingRules.h b/Engine/Plugins/Developer/Concert/ConcertUI/ConcertClientSharedSlate/Source/ConcertClientSharedSlate/Public/Replication/Settings/ConcertPerClassSubobjectMatchingRules.h
new file mode 100644
--- /dev/null
+++ b/Engine/Plugins/Developer/Concert/ConcertUI/ConcertClientSharedSlate/Source/ConcertClientSharedSlate/Public/Replication/Settings/ConcertPerClassSubobjectMatchingRules.h
@@ -0,0 +1,39 @@
+// Copyright Epic Games, Inc. All Rights Reserved.
+
+#pragma once
+
+#include "ConcertInheritableClassOption.h"
+#include "ConcertSubobjectMatchingRules.h"
+#include "ConcertPerClassSubobjectMatchingRules.generated.h"
+
+USTRUCT()
+struct CONCERTCLIENTSHAREDSLATE_API FConcertInheritableSubobjectMatchingRules : public FConcertSubobjectMatchingRules, public FConcertInheritableClassOption
+{
+	GENERATED_BODY()
+};
+
+/**
+ * Binds subobject matching rules to classes.
+ * The matching rules can then be applied recursively, if FConcertSubobjectDefaultSelectionRules::bInheritFromBase is set to true.
+ *
+ * For example, suppose you have set up two sets of rules: one for matching subobjects of an UStaticMeshComponent and one for USkeletalMeshComponent.
+ * If the bInheritFromBase is set true for USkeletalMeshComponent's rules, then UStaticMeshComponent will also be applied when searching a skeletal mesh component.
+ */
+USTRUCT()
+struct CONCERTCLIENTSHAREDSLATE_API FConcertPerClassSubobjectMatchingRules
+{
+	GENERATED_BODY()
+	
+	/**
+	 * Maps matching rules to actor and component classes.
+	 * 
+	 * Examples:
+	 *  - Whenever you add a StaticMeshActor actor, you its static mesh component to be added automatically, too (set IncludeSubobjectRegex to the components name).
+	 *  - Whenever you add a component, you want to add its child components, too (set IncludeAllOption = AllComponents) .
+	 */
+	UPROPERTY(EditAnywhere, Config, Category = "Replication|Editor")
+	TMap<FSoftClassPath, FConcertInheritableSubobjectMatchingRules> SubobjectMatchingRules;
+
+	/** Recursively walks up Object class hierarchy and tries to find subobjects  */
+	void MatchSubobjectsRecursivelyFor(const UObject& Object, TFunctionRef<void(UObject&)> OnSubobjectMatched) const;
+};
\ No newline at end of file
diff --git a/Engine/Plugins/Developer/Concert/ConcertUI/ConcertClientSharedSlate/Source/ConcertClientSharedSlate/Public/Replication/Settings/ConcertReplicationEditorSettings.h b/Engine/Plugins/Developer/Concert/ConcertUI/ConcertClientSharedSlate/Source/ConcertClientSharedSlate/Public/Replication/Settings/ConcertReplicationEditorSettings.h
deleted file mode 100644
diff --git a/Engine/Plugins/Developer/Concert/ConcertUI/ConcertClientSharedSlate/Source/ConcertClientSharedSlate/Public/Replication/Settings/ConcertStreamObjectAutoBindingRules.h b/Engine/Plugins/Developer/Concert/ConcertUI/ConcertClientSharedSlate/Source/ConcertClientSharedSlate/Public/Replication/Settings/ConcertStreamObjectAutoBindingRules.h
new file mode 100644
--- /dev/null
+++ b/Engine/Plugins/Developer/Concert/ConcertUI/ConcertClientSharedSlate/Source/ConcertClientSharedSlate/Public/Replication/Settings/ConcertStreamObjectAutoBindingRules.h
@@ -0,0 +1,54 @@
+// Copyright Epic Games, Inc. All Rights Reserved.
+
+#pragma once
+
+#include "ConcertInheritableClassOption.h"
+#include "ConcertPerClassSubobjectMatchingRules.h"
+#include "Templates/Function.h"
+#include "ConcertStreamObjectAutoBindingRules.generated.h"
+
+class UObject;
+struct FConcertPropertyChain;
+struct FReplicatedObjectInfo;
+
+USTRUCT()
+struct CONCERTCLIENTSHAREDSLATE_API FConcertDefaultPropertySelection : public FConcertInheritableClassOption
+{
+	GENERATED_BODY()
+
+	/**
+	 * A list of properties that should be selected by default when you add a new object type.
+	 * Specify properties by using the names as the appear in the replication editor and each property with ".".
+	 * Example: "RelativeLocation.X"
+	 */
+	UPROPERTY(EditAnywhere, Category = "Config")
+	TArray<FString> DefaultSelectedProperties;
+};
+
+/**
+ * When a user adds an object via the stream editor, you may want to automatically bind properties and add additional subobjects
+ * This structure contains rules to achieve this.
+ */
+USTRUCT()
+struct CONCERTCLIENTSHAREDSLATE_API FConcertStreamObjectAutoBindingRules
+{
+	GENERATED_BODY()
+	
+	/** Properties you want selected by default when you add a new replicated object in the editor */
+	UPROPERTY(EditAnywhere, Category = "Replication|Editor")
+	TMap<FSoftClassPath, FConcertDefaultPropertySelection> DefaultPropertySelection;
+
+	/**
+	 * Rules for auto adding additional subobjects.
+	 * 
+	 * For example, you can set this up so when you add a static mesh actor, its static mesh component is automatically added, too.
+	 * The simplest way to achieve this is to bind AStaticMeshActor class to a setting with IncludeAllOption set to AllComponents. 
+	 * 
+	 * Matched subobjects will be auto added, too.
+	 */
+	UPROPERTY(EditAnywhere, Category = "Replication|Editor")
+	FConcertPerClassSubobjectMatchingRules DefaultAddedSubobjectRules;
+
+	/** Reads DefaultPropertySelection and calls Callback for any default property selections based on the Class just added. */
+	void AddDefaultPropertiesFromSettings(UClass& Class, TFunctionRef<void(FConcertPropertyChain&& Chain)> Callback) const;
+};
\ No newline at end of file
diff --git a/Engine/Plugins/Developer/Concert/ConcertUI/ConcertClientSharedSlate/Source/ConcertClientSharedSlate/Public/Replication/Settings/ConcertSubobjectMatchingRules.h b/Engine/Plugins/Developer/Concert/ConcertUI/ConcertClientSharedSlate/Source/ConcertClientSharedSlate/Public/Replication/Settings/ConcertSubobjectMatchingRules.h
new file mode 100644
--- /dev/null
+++ b/Engine/Plugins/Developer/Concert/ConcertUI/ConcertClientSharedSlate/Source/ConcertClientSharedSlate/Public/Replication/Settings/ConcertSubobjectMatchingRules.h
@@ -0,0 +1,62 @@
+// Copyright Epic Games, Inc. All Rights Reserved.
+
+#pragma once
+
+#include "Templates/Function.h"
+#include "UObject/SoftObjectPath.h"
+#include "ConcertSubobjectMatchingRules.generated.h"
+
+UENUM()
+enum class EConcertIncludeAllSubobjectsType : uint8
+{
+	/** Include nothing by default */
+	None,
+	/** Include all components */
+	AllComponents,
+	/** Include all subobjects recursively (includes components, their subobjects, etc. ) */
+	AllSubobjects
+};
+
+/** Defines rules for finding subobjects of an object. */
+USTRUCT()
+struct FConcertSubobjectMatchingRules
+{
+	GENERATED_BODY()
+	
+	/**
+	 * A list of subobject classes to add by default.
+	 * For example, when you add an actor, you want to also add certain component classes.
+	 */
+	UPROPERTY(EditAnywhere, Category = "Config")
+	TSet<FSoftClassPath> IncludeClasses;
+
+	/**
+	 * Components that match any of this regex will be excluded from IncludeClasses.
+	 *
+	 * Tips:
+	 * - If you want to include "Component" but not "ComponentName" you can use boundaries "\bComponent\b".
+	 * . ".*" matches any sequence of characters.
+	 */
+	UPROPERTY(EditAnywhere, Category = "Config")
+	TSet<FString> ExcludeSubobjectRegex;
+
+	/**
+	 * Components that match any of this regex will be included as well.
+	 * 
+	 * Tips:
+	 * - If you want to include "Component" but not "ComponentName" you can use boundaries "\bComponent\b".
+	 * . ".*" matches any sequence of characters. 
+	 */
+	UPROPERTY(EditAnywhere, Category = "Config")
+	TSet<FString> IncludeSubobjectRegex;
+
+	/** Behaviour to configure for including all of a certain type of subobject. */
+	UPROPERTY(EditAnywhere, AdvancedDisplay, Category = "Config")
+	EConcertIncludeAllSubobjectsType IncludeAllOption = EConcertIncludeAllSubobjectsType::None;
+
+	/**
+	 * Looks for subobjects in AddedObject that match these rules.
+	 * This looks for direct subobjects only though you can call MatchToSubobjectsIn recursively on what you received through OnSubobjectMatched.
+	 */
+	void MatchToSubobjectsIn(const UObject& AddedObject, TFunctionRef<void(UObject&)> OnSubobjectMatched) const;
+};
