Move client independent replication UI code to ConcertSharedSlate. This to future proof the code in case we want replication UI e.g. on the server.

diff --git a/Engine/Plugins/Developer/Concert/ConcertApp/MultiUserClient/Source/MultiUserClient/Private/Replication/Client/ReplicationClient.cpp b/Engine/Plugins/Developer/Concert/ConcertApp/MultiUserClient/Source/MultiUserClient/Private/Replication/Client/ReplicationClient.cpp
--- a/Engine/Plugins/Developer/Concert/ConcertApp/MultiUserClient/Source/MultiUserClient/Private/Replication/Client/ReplicationClient.cpp
+++ b/Engine/Plugins/Developer/Concert/ConcertApp/MultiUserClient/Source/MultiUserClient/Private/Replication/Client/ReplicationClient.cpp
@@ -8,9 +8,10 @@
 #include "Replication/Stream/Discovery/MultiUserStreamExtender.h"
 #include "Replication/Stream/StreamChangeTracker.h"
 #include "Replication/Submission/Data/AuthoritySubmission.h"
+#include "Settings/MultiUserReplicationSettings.h"
 
 #include "Misc/CoreDelegates.h"
-#include "Settings/MultiUserReplicationSettings.h"
+#include "Replication/ClientReplicationWidgetFactories.h"
 
 namespace UE::MultiUserClient
 {
@@ -34,10 +35,14 @@ namespace UE::MultiUserClient
 			FGetStreamContent::CreateLambda([this](){ return &StreamSynchronizer->GetServerState(); }),
 			SubmissionQueue
 			)
-		, LocalClientEditModel(CreatePropertySelectionModel(
-			*ClientContentStorage->Stream,
-			ClientContentStorage->Stream->MakeReplicationMapGetterAttribute(),
-			MakeShared<FMultiUserStreamExtender>(EndpointId, InDiscoveryContainer))
+		, LocalClientEditModel(
+			ConcertClientSharedSlate::CreateTransactionalStreamModel(
+				CreateBaseStreamModel(
+					ClientContentStorage->Stream->MakeReplicationMapGetterAttribute(),
+					MakeShared<FMultiUserStreamExtender>(EndpointId, InDiscoveryContainer)
+					),
+				*ClientContentStorage->Stream
+				)
 			)
 		, LocalClientStreamDiffer(
 			GetStreamSynchronizer(),
@@ -70,7 +75,7 @@ namespace UE::MultiUserClient
 	void FReplicationClient::OnObjectsChanged(
 		TConstArrayView<UObject*> AddedObjects,
 		TConstArrayView<FSoftObjectPath> RemovedObjects,
-		ConcertClientSharedSlate::EReplicatedObjectChangeReason ReplicatedObjectChangeReason
+		ConcertSharedSlate::EReplicatedObjectChangeReason ReplicatedObjectChangeReason
 		)
 	{
 		DeferOnModelChanged(AddedObjects);
diff --git a/Engine/Plugins/Developer/Concert/ConcertApp/MultiUserClient/Source/MultiUserClient/Private/Replication/Client/ReplicationClient.h b/Engine/Plugins/Developer/Concert/ConcertApp/MultiUserClient/Source/MultiUserClient/Private/Replication/Client/ReplicationClient.h
--- a/Engine/Plugins/Developer/Concert/ConcertApp/MultiUserClient/Source/MultiUserClient/Private/Replication/Client/ReplicationClient.h
+++ b/Engine/Plugins/Developer/Concert/ConcertApp/MultiUserClient/Source/MultiUserClient/Private/Replication/Client/ReplicationClient.h
@@ -19,7 +19,7 @@
 class IConcertClient;
 class UMultiUserReplicationClientPreset;
 
-namespace UE::ConcertClientSharedSlate
+namespace UE::ConcertSharedSlate
 {
 	class IEditableReplicationStreamModel;
 	enum class EReplicatedObjectChangeReason : uint8;
@@ -67,7 +67,7 @@ namespace UE::MultiUserClient
 		 * Listen for events on the owning FReplicationClientManager::OnPreRemoteClientRemoved and FMultiUserReplicationManager::OnLeaveSession
 		 * @see FReplicationClientManager and FMultiUserReplicationManager
 		 */
-		TSharedRef<ConcertClientSharedSlate::IEditableReplicationStreamModel> GetClientEditModel() const { return LocalClientEditModel; }
+		TSharedRef<ConcertSharedSlate::IEditableReplicationStreamModel> GetClientEditModel() const { return LocalClientEditModel; }
 		IClientStreamSynchronizer& GetStreamSynchronizer() const { return *StreamSynchronizer.Get(); }
 		IClientAuthoritySynchronizer& GetAuthoritySynchronizer() const { return *AuthoritySynchronizer.Get(); }
 		
@@ -131,7 +131,7 @@ namespace UE::MultiUserClient
 		 * The UI is destroyed right before FReplicationClient, which is guaranteed by FMultiUserReplicationManager.
 		 * @see FMultiUserReplicationManager::OnLeaveSession
 		 */
-		TSharedRef<ConcertClientSharedSlate::IEditableReplicationStreamModel> LocalClientEditModel;
+		TSharedRef<ConcertSharedSlate::IEditableReplicationStreamModel> LocalClientEditModel;
 		
 		/** Tracks changes made to server's state of the client's streams and prepares to upload them using StreamSynchronizer. */
 		FStreamChangeTracker LocalClientStreamDiffer;
@@ -160,7 +160,7 @@ namespace UE::MultiUserClient
 		TOptional<FDeferredOnModelChangedData> DeferredOnModelChangedData;
 
 		// Respond to model changing
-		void OnObjectsChanged(TConstArrayView<UObject*> AddedObjects, TConstArrayView<FSoftObjectPath> RemovedObjects, ConcertClientSharedSlate::EReplicatedObjectChangeReason ReplicatedObjectChangeReason);
+		void OnObjectsChanged(TConstArrayView<UObject*> AddedObjects, TConstArrayView<FSoftObjectPath> RemovedObjects, ConcertSharedSlate::EReplicatedObjectChangeReason ReplicatedObjectChangeReason);
 		void OnPropertiesChanged();
 
 		void OnServerStateChanged();
diff --git a/Engine/Plugins/Developer/Concert/ConcertApp/MultiUserClient/Source/MultiUserClient/Private/Replication/Client/ReplicationClientManager.h b/Engine/Plugins/Developer/Concert/ConcertApp/MultiUserClient/Source/MultiUserClient/Private/Replication/Client/ReplicationClientManager.h
--- a/Engine/Plugins/Developer/Concert/ConcertApp/MultiUserClient/Source/MultiUserClient/Private/Replication/Client/ReplicationClientManager.h
+++ b/Engine/Plugins/Developer/Concert/ConcertApp/MultiUserClient/Source/MultiUserClient/Private/Replication/Client/ReplicationClientManager.h
@@ -22,7 +22,7 @@ enum class EConcertClientStatus : uint8;
 
 struct FConcertSessionClientInfo;
 
-namespace UE::ConcertClientSharedSlate
+namespace UE::ConcertSharedSlate
 {
 	class IEditableReplicationStreamModel;
 }
diff --git a/Engine/Plugins/Developer/Concert/ConcertApp/MultiUserClient/Source/MultiUserClient/Private/Replication/MultiUserReplicationManager.h b/Engine/Plugins/Developer/Concert/ConcertApp/MultiUserClient/Source/MultiUserClient/Private/Replication/MultiUserReplicationManager.h
--- a/Engine/Plugins/Developer/Concert/ConcertApp/MultiUserClient/Source/MultiUserClient/Private/Replication/MultiUserReplicationManager.h
+++ b/Engine/Plugins/Developer/Concert/ConcertApp/MultiUserClient/Source/MultiUserClient/Private/Replication/MultiUserReplicationManager.h
@@ -20,7 +20,7 @@ namespace UE::ConcertSyncClient::Replication
 	struct FJoinReplicatedSessionResult;
 }
 
-namespace UE::ConcertClientSharedSlate
+namespace UE::ConcertSharedSlate
 {
 	class IEditableReplicationStreamModel;
 }
diff --git a/Engine/Plugins/Developer/Concert/ConcertApp/MultiUserClient/Source/MultiUserClient/Private/Replication/Stream/Discovery/MultiUserStreamExtender.cpp b/Engine/Plugins/Developer/Concert/ConcertApp/MultiUserClient/Source/MultiUserClient/Private/Replication/Stream/Discovery/MultiUserStreamExtender.cpp
--- a/Engine/Plugins/Developer/Concert/ConcertApp/MultiUserClient/Source/MultiUserClient/Private/Replication/Stream/Discovery/MultiUserStreamExtender.cpp
+++ b/Engine/Plugins/Developer/Concert/ConcertApp/MultiUserClient/Source/MultiUserClient/Private/Replication/Stream/Discovery/MultiUserStreamExtender.cpp
@@ -21,20 +21,20 @@ namespace UE::MultiUserClient
 		, RegisteredExtenders(InRegisteredExtenders)
 	{}
 
-	void FMultiUserStreamExtender::ExtendStream(UObject& ExtendedObject, ConcertClientSharedSlate::IStreamExtensionContext& Context)
+	void FMultiUserStreamExtender::ExtendStream(UObject& ExtendedObject, ConcertSharedSlate::IStreamExtensionContext& Context)
 	{
 		ExtendBySettings.ExtendStream(ExtendedObject, Context);
 		ExtendStreamWithRegisteredDiscoverers(ExtendedObject, Context);
 	}
 	
-	void FMultiUserStreamExtender::ExtendStreamWithRegisteredDiscoverers(UObject& ExtendedObject, ConcertClientSharedSlate::IStreamExtensionContext& Context) const
+	void FMultiUserStreamExtender::ExtendStreamWithRegisteredDiscoverers(UObject& ExtendedObject, ConcertSharedSlate::IStreamExtensionContext& Context) const
 	{
 		class FDiscoveryContext : public IReplicationDiscoveryContext
 		{
-			ConcertClientSharedSlate::IStreamExtensionContext& Context;
+			ConcertSharedSlate::IStreamExtensionContext& Context;
 		public:
 			
-			explicit FDiscoveryContext(ConcertClientSharedSlate::IStreamExtensionContext& InContext)
+			explicit FDiscoveryContext(ConcertSharedSlate::IStreamExtensionContext& InContext)
 				: Context(InContext)
 			{}
 
diff --git a/Engine/Plugins/Developer/Concert/ConcertApp/MultiUserClient/Source/MultiUserClient/Private/Replication/Stream/Discovery/MultiUserStreamExtender.h b/Engine/Plugins/Developer/Concert/ConcertApp/MultiUserClient/Source/MultiUserClient/Private/Replication/Stream/Discovery/MultiUserStreamExtender.h
--- a/Engine/Plugins/Developer/Concert/ConcertApp/MultiUserClient/Source/MultiUserClient/Private/Replication/Stream/Discovery/MultiUserStreamExtender.h
+++ b/Engine/Plugins/Developer/Concert/ConcertApp/MultiUserClient/Source/MultiUserClient/Private/Replication/Stream/Discovery/MultiUserStreamExtender.h
@@ -3,7 +3,7 @@
 #pragma once
 
 #include "Replication/Editor/Model/Extension/IStreamExtender.h"
-#include "Replication/Editor/Model/StreamExtenderBySettings.h"
+#include "Replication/Editor/Model/Extension/StreamExtenderBySettings.h"
 
 #include "Misc/Guid.h"
 
@@ -19,7 +19,7 @@ namespace UE::MultiUserClient
 	 *
 	 * Every FReplicationClient owns an instance of this and executes it when an object is added for the client by the local editor.
 	 */
-	class FMultiUserStreamExtender : public ConcertClientSharedSlate::IStreamExtender
+	class FMultiUserStreamExtender : public ConcertSharedSlate::IStreamExtender
 	{
 	public:
 		
@@ -30,7 +30,7 @@ namespace UE::MultiUserClient
 		FMultiUserStreamExtender(const FGuid& InClientId, FReplicationDiscoveryContainer& InRegisteredExtenders);
 
 		//~ Begin IStreamExtender Interface
-		virtual void ExtendStream(UObject& ExtendedObject, ConcertClientSharedSlate::IStreamExtensionContext& Context) override;
+		virtual void ExtendStream(UObject& ExtendedObject, ConcertSharedSlate::IStreamExtensionContext& Context) override;
 		//~ End IStreamExtender Interface
 
 	private:
@@ -44,7 +44,7 @@ namespace UE::MultiUserClient
 		/** Allows generic extenders to extend added objects. These can be added via the IMultiUserReplication interface. */
 		FReplicationDiscoveryContainer& RegisteredExtenders;
 		
-		void ExtendStreamWithRegisteredDiscoverers(UObject& ExtendedObject, ConcertClientSharedSlate::IStreamExtensionContext& Context) const;
+		void ExtendStreamWithRegisteredDiscoverers(UObject& ExtendedObject, ConcertSharedSlate::IStreamExtensionContext& Context) const;
 	};
 }
 
diff --git a/Engine/Plugins/Developer/Concert/ConcertApp/MultiUserClient/Source/MultiUserClient/Private/Replication/Submission/AutoSubmissionPolicy.cpp b/Engine/Plugins/Developer/Concert/ConcertApp/MultiUserClient/Source/MultiUserClient/Private/Replication/Submission/AutoSubmissionPolicy.cpp
--- a/Engine/Plugins/Developer/Concert/ConcertApp/MultiUserClient/Source/MultiUserClient/Private/Replication/Submission/AutoSubmissionPolicy.cpp
+++ b/Engine/Plugins/Developer/Concert/ConcertApp/MultiUserClient/Source/MultiUserClient/Private/Replication/Submission/AutoSubmissionPolicy.cpp
@@ -14,7 +14,7 @@ namespace UE::MultiUserClient
 	FAutoSubmissionPolicy::FAutoSubmissionPolicy(
 		FSubmissionQueue& InSubmissionQueue,
 		const FChangeRequestBuilder& InRequestBuilder,
-		ConcertClientSharedSlate::IEditableReplicationStreamModel& InStreamEditorModel,
+		ConcertSharedSlate::IEditableReplicationStreamModel& InStreamEditorModel,
 		FAuthorityChangeTracker& InAuthorityChangeTracker,
 		FFrequencyChangeTracker& InFrequencyChangeTracker
 		)
diff --git a/Engine/Plugins/Developer/Concert/ConcertApp/MultiUserClient/Source/MultiUserClient/Private/Replication/Submission/AutoSubmissionPolicy.h b/Engine/Plugins/Developer/Concert/ConcertApp/MultiUserClient/Source/MultiUserClient/Private/Replication/Submission/AutoSubmissionPolicy.h
--- a/Engine/Plugins/Developer/Concert/ConcertApp/MultiUserClient/Source/MultiUserClient/Private/Replication/Submission/AutoSubmissionPolicy.h
+++ b/Engine/Plugins/Developer/Concert/ConcertApp/MultiUserClient/Source/MultiUserClient/Private/Replication/Submission/AutoSubmissionPolicy.h
@@ -11,7 +11,7 @@ namespace UE::MultiUserClient
 	class FFrequencyChangeTracker;
 }
 
-namespace UE::ConcertClientSharedSlate
+namespace UE::ConcertSharedSlate
 {
 	class IEditableReplicationStreamModel;
 }
@@ -40,7 +40,7 @@ namespace UE::MultiUserClient
 		FAutoSubmissionPolicy(
 			FSubmissionQueue& InSubmissionQueue UE_LIFETIMEBOUND,
 			const FChangeRequestBuilder& InRequestBuilder UE_LIFETIMEBOUND,
-			ConcertClientSharedSlate::IEditableReplicationStreamModel& InStreamEditorModel UE_LIFETIMEBOUND,
+			ConcertSharedSlate::IEditableReplicationStreamModel& InStreamEditorModel UE_LIFETIMEBOUND,
 			FAuthorityChangeTracker& InAuthorityChangeTracker UE_LIFETIMEBOUND,
 			FFrequencyChangeTracker& InFrequencyChangeTracker UE_LIFETIMEBOUND
 			);
@@ -57,7 +57,7 @@ namespace UE::MultiUserClient
 		const FChangeRequestBuilder& RequestBuilder;
 
 		/** Informs us when the stream is structurally changed by the user. */
-		ConcertClientSharedSlate::IEditableReplicationStreamModel& StreamEditorModel;
+		ConcertSharedSlate::IEditableReplicationStreamModel& StreamEditorModel;
 		/** Informs us when authority is changed by the user. */
 		FAuthorityChangeTracker& AuthorityChangeTracker;
 		/** Informs us when frequency settings are changed by the user. */
@@ -70,7 +70,7 @@ namespace UE::MultiUserClient
 		virtual void PerformSubmission_GameThread(ISubmissionWorkflow& Workflow) override;
 		//~ End FSelfUnregisteringDeferredSubmitter Interface
 
-		void OnObjectsChanged(TArrayView<UObject* const>, TArrayView<const FSoftObjectPath>, ConcertClientSharedSlate::EReplicatedObjectChangeReason) { OnChangesDetected(); }
+		void OnObjectsChanged(TArrayView<UObject* const>, TArrayView<const FSoftObjectPath>, ConcertSharedSlate::EReplicatedObjectChangeReason) { OnChangesDetected(); }
 		void OnChangesDetected();
 	};
 }
diff --git a/Engine/Plugins/Developer/Concert/ConcertApp/MultiUserClient/Source/MultiUserClient/Private/Widgets/ActiveSession/Replication/Client/Multi/Columns/MultiStreamColumns.h b/Engine/Plugins/Developer/Concert/ConcertApp/MultiUserClient/Source/MultiUserClient/Private/Widgets/ActiveSession/Replication/Client/Multi/Columns/MultiStreamColumns.h
--- a/Engine/Plugins/Developer/Concert/ConcertApp/MultiUserClient/Source/MultiUserClient/Private/Widgets/ActiveSession/Replication/Client/Multi/Columns/MultiStreamColumns.h
+++ b/Engine/Plugins/Developer/Concert/ConcertApp/MultiUserClient/Source/MultiUserClient/Private/Widgets/ActiveSession/Replication/Client/Multi/Columns/MultiStreamColumns.h
@@ -15,7 +15,7 @@ namespace UE::MultiUserClient
 	class FReplicationClientManager;
 }
 
-namespace UE::ConcertClientSharedSlate
+namespace UE::ConcertSharedSlate
 {
 	class IReplicationStreamModel;
 	class IMultiReplicationStreamEditor;
@@ -46,9 +46,9 @@ namespace UE::MultiUserClient::MultiStreamColumns
 	 * 
 	 * @return A checkbox for controlling the authority of the object in the row
 	 */
-	ConcertClientSharedSlate::ReplicationColumns::FReplicationTopLevelObjectColumn ReplicationToggle(
+	ConcertSharedSlate::ReplicationColumns::FReplicationTopLevelObjectColumn ReplicationToggle(
 		TSharedRef<IConcertClient> ConcertClient,
-		TAttribute<ConcertClientSharedSlate::IReplicationStreamModel*> ConsolidatedStreamModelAttribute,
+		TAttribute<ConcertSharedSlate::IReplicationStreamModel*> ConsolidatedStreamModelAttribute,
 		FReplicationClientManager& ClientManager,
 		const int32 ColumnsSortPriority = static_cast<int32>(EColumnSortOrder::ReplicationToggle)
 		);
@@ -62,9 +62,9 @@ namespace UE::MultiUserClient::MultiStreamColumns
 	 * @param ClientManager Used to access all clients for display in the combo box drop-down
 	 * @param ColumnsSortPriority The order relative to the other columns
 	 */
-	ConcertClientSharedSlate::ReplicationColumns::FReplicationTopLevelObjectColumn ReassignOwnership(
+	ConcertSharedSlate::ReplicationColumns::FReplicationTopLevelObjectColumn ReassignOwnership(
 		TSharedRef<IConcertClient> ConcertClient,
-		TAttribute<TSharedPtr<ConcertClientSharedSlate::IMultiReplicationStreamEditor>> MultiStreamModelAttribute,
+		TAttribute<TSharedPtr<ConcertSharedSlate::IMultiReplicationStreamEditor>> MultiStreamModelAttribute,
 		FReassignObjectPropertiesLogic& ReassignmentLogic,
 		const FReplicationClientManager& ClientManager,
 		int32 ColumnsSortPriority = static_cast<int32>(EColumnSortOrder::ReassignOwnership)
@@ -80,8 +80,8 @@ namespace UE::MultiUserClient::MultiStreamColumns
 	 * 
 	 * @return A column that spawns a combo box for assigning properties
 	 */
-	ConcertClientSharedSlate::ReplicationColumns::FReplicationPropertyColumn AssignPropertyColumn(
-		TAttribute<TSharedPtr<ConcertClientSharedSlate::IMultiReplicationStreamEditor>> MultiStreamEditor,
+	ConcertSharedSlate::ReplicationColumns::FReplicationPropertyColumn AssignPropertyColumn(
+		TAttribute<TSharedPtr<ConcertSharedSlate::IMultiReplicationStreamEditor>> MultiStreamEditor,
 		TSharedRef<IConcertClient> ConcertClient,
 		FReplicationClientManager& ClientManager,
 		const int32 ColumnsSortPriority = static_cast<int32>(EColumnSortOrder::AssignPropertyColumn)
diff --git a/Engine/Plugins/Developer/Concert/ConcertApp/MultiUserClient/Source/MultiUserClient/Private/Widgets/ActiveSession/Replication/Client/Multi/Columns/MultiStreamColumns_AssignProperties.cpp b/Engine/Plugins/Developer/Concert/ConcertApp/MultiUserClient/Source/MultiUserClient/Private/Widgets/ActiveSession/Replication/Client/Multi/Columns/MultiStreamColumns_AssignProperties.cpp
--- a/Engine/Plugins/Developer/Concert/ConcertApp/MultiUserClient/Source/MultiUserClient/Private/Widgets/ActiveSession/Replication/Client/Multi/Columns/MultiStreamColumns_AssignProperties.cpp
+++ b/Engine/Plugins/Developer/Concert/ConcertApp/MultiUserClient/Source/MultiUserClient/Private/Widgets/ActiveSession/Replication/Client/Multi/Columns/MultiStreamColumns_AssignProperties.cpp
@@ -22,12 +22,12 @@ namespace UE::MultiUserClient::MultiStreamColumns
 	namespace AssignPropertyColumnUtils
 	{
 		static void ForEachStreamAssignedTo(
-			const ConcertClientSharedSlate::IMultiReplicationStreamEditor& MultiEditor,
+			const ConcertSharedSlate::IMultiReplicationStreamEditor& MultiEditor,
 			const FConcertPropertyChain& Property,
-			TFunctionRef<void(const TSharedRef<ConcertClientSharedSlate::IEditableReplicationStreamModel>& Stream)> Consume
+			TFunctionRef<void(const TSharedRef<ConcertSharedSlate::IEditableReplicationStreamModel>& Stream)> Consume
 		)
 		{
-			for (const TSharedRef<ConcertClientSharedSlate::IEditableReplicationStreamModel>& Stream : MultiEditor.GetMultiStreamModel().GetEditableStreams())
+			for (const TSharedRef<ConcertSharedSlate::IEditableReplicationStreamModel>& Stream : MultiEditor.GetMultiStreamModel().GetEditableStreams())
 			{
 				for (const FSoftObjectPath& SelectedObject : MultiEditor.GetEditorBase().GetObjectsBeingPropertyEdited())
 				{
@@ -42,7 +42,7 @@ namespace UE::MultiUserClient::MultiStreamColumns
 			}
 		}
 
-		const FReplicationClient* FindClientByStream(const FReplicationClientManager& ClientManager, const ConcertClientSharedSlate::IReplicationStreamModel& StreamModel)
+		const FReplicationClient* FindClientByStream(const FReplicationClientManager& ClientManager, const ConcertSharedSlate::IReplicationStreamModel& StreamModel)
 		{
 			if (&ClientManager.GetLocalClient().GetClientEditModel().Get() == &StreamModel)
 			{
@@ -60,7 +60,7 @@ namespace UE::MultiUserClient::MultiStreamColumns
 			return nullptr;
 		}
 		
-		static FString GetClientDisplayText(const IConcertClient& InConcertClient, const FReplicationClientManager& ClientManager, const ConcertClientSharedSlate::IReplicationStreamModel& StreamModel)
+		static FString GetClientDisplayText(const IConcertClient& InConcertClient, const FReplicationClientManager& ClientManager, const ConcertSharedSlate::IReplicationStreamModel& StreamModel)
 		{
 			if (const FReplicationClient* Client = FindClientByStream(ClientManager, StreamModel))
 			{
@@ -72,15 +72,15 @@ namespace UE::MultiUserClient::MultiStreamColumns
 		}
 	}
 	
-	ConcertClientSharedSlate::ReplicationColumns::FReplicationPropertyColumn AssignPropertyColumn(
-		TAttribute<TSharedPtr<ConcertClientSharedSlate::IMultiReplicationStreamEditor>> MultiStreamEditor,
+	ConcertSharedSlate::ReplicationColumns::FReplicationPropertyColumn AssignPropertyColumn(
+		TAttribute<TSharedPtr<ConcertSharedSlate::IMultiReplicationStreamEditor>> MultiStreamEditor,
 		TSharedRef<IConcertClient> ConcertClient,
 		FReplicationClientManager& ClientManager,
 		const int32 ColumnsSortPriority
 		)
 	{
-		using namespace ConcertClientSharedSlate;
-		using namespace ConcertClientSharedSlate::ReplicationColumns;
+		using namespace ConcertSharedSlate;
+		using namespace ConcertSharedSlate::ReplicationColumns;
 		check(MultiStreamEditor.IsBound() || MultiStreamEditor.IsSet());
 
 		const auto PopulateSearch = [MultiStreamEditor, &ClientManager, ConcertClient](const FReplicatedPropertyData& Data, TArray<FString>& InOutSearchStrings)
diff --git a/Engine/Plugins/Developer/Concert/ConcertApp/MultiUserClient/Source/MultiUserClient/Private/Widgets/ActiveSession/Replication/Client/Multi/Columns/MultiStreamColumns_ReassignOwnership.cpp b/Engine/Plugins/Developer/Concert/ConcertApp/MultiUserClient/Source/MultiUserClient/Private/Widgets/ActiveSession/Replication/Client/Multi/Columns/MultiStreamColumns_ReassignOwnership.cpp
--- a/Engine/Plugins/Developer/Concert/ConcertApp/MultiUserClient/Source/MultiUserClient/Private/Widgets/ActiveSession/Replication/Client/Multi/Columns/MultiStreamColumns_ReassignOwnership.cpp
+++ b/Engine/Plugins/Developer/Concert/ConcertApp/MultiUserClient/Source/MultiUserClient/Private/Widgets/ActiveSession/Replication/Client/Multi/Columns/MultiStreamColumns_ReassignOwnership.cpp
@@ -14,15 +14,15 @@ namespace UE::MultiUserClient::MultiStreamColumns
 {
 	const FName ReassignOwnershipColumnId(TEXT("ReassignOwnershipColumn"));
 	
-	ConcertClientSharedSlate::ReplicationColumns::FReplicationTopLevelObjectColumn ReassignOwnership(
+	ConcertSharedSlate::ReplicationColumns::FReplicationTopLevelObjectColumn ReassignOwnership(
 		TSharedRef<IConcertClient> ConcertClient,
-		TAttribute<TSharedPtr<ConcertClientSharedSlate::IMultiReplicationStreamEditor>> MultiStreamModelAttribute,
+		TAttribute<TSharedPtr<ConcertSharedSlate::IMultiReplicationStreamEditor>> MultiStreamModelAttribute,
 		FReassignObjectPropertiesLogic& ReassignmentLogic,
 		const FReplicationClientManager& ClientManager,
 		const int32 ColumnsSortPriority
 		)
 	{
-		using namespace ConcertClientSharedSlate::ReplicationColumns;
+		using namespace ConcertSharedSlate::ReplicationColumns;
 
 		auto MakeWidget =
 			[ConcertClient, MultiStreamModelAttribute = MoveTemp(MultiStreamModelAttribute), &ReassignmentLogic, &ClientManager]
@@ -32,7 +32,7 @@ namespace UE::MultiUserClient::MultiStreamColumns
 					.ManagedObject(InArgs.RowData.GetObjectPath())
 					.ConsolidatedModel_Lambda([MultiStreamModelAttribute]()
 					{
-						const TSharedPtr<ConcertClientSharedSlate::IMultiReplicationStreamEditor> Editor = MultiStreamModelAttribute.Get();
+						const TSharedPtr<ConcertSharedSlate::IMultiReplicationStreamEditor> Editor = MultiStreamModelAttribute.Get();
 						return Editor
 							? &MultiStreamModelAttribute.Get()->GetConsolidatedModel()
 							: nullptr;
@@ -40,13 +40,13 @@ namespace UE::MultiUserClient::MultiStreamColumns
 					.HighlightText(InArgs.HighlightText)
 					.OnReassignAllOptionClicked_Lambda([MultiStreamModelAttribute](auto)
 					{
-						if (const TSharedPtr<ConcertClientSharedSlate::IMultiReplicationStreamEditor> Model = MultiStreamModelAttribute.Get())
+						if (const TSharedPtr<ConcertSharedSlate::IMultiReplicationStreamEditor> Model = MultiStreamModelAttribute.Get())
 						{
 							Model->GetEditorBase().RequestObjectColumnResort(ReassignOwnershipColumnId);
 						}
 					});
 			};
-		auto IsLessThan = [ConcertClient, &ReassignmentLogic](const ConcertClientSharedSlate::FReplicatedObjectData& Left, const ConcertClientSharedSlate::FReplicatedObjectData& Right)
+		auto IsLessThan = [ConcertClient, &ReassignmentLogic](const ConcertSharedSlate::FReplicatedObjectData& Left, const ConcertSharedSlate::FReplicatedObjectData& Right)
 		{
 			const TOptional<FString> LeftClientDisplayString = SReassignObjectComboBox::GetDisplayString(ConcertClient, ReassignmentLogic, Left.GetObjectPath());
 			const TOptional<FString> RightClientDisplayString = SReassignObjectComboBox::GetDisplayString(ConcertClient, ReassignmentLogic, Right.GetObjectPath());
@@ -62,7 +62,7 @@ namespace UE::MultiUserClient::MultiStreamColumns
 		return FReplicationTopLevelObjectColumn(
 			FReplicationTopLevelObjectColumn::FArguments()
 				.GenerateWidgetColumn_Lambda(MoveTemp(MakeWidget))
-				.PopulateSearchItems_Lambda([ConcertClient, &ReassignmentLogic](const ConcertClientSharedSlate::FReplicatedObjectData& ObjectData, TArray<FString>& InOutSearchStrings)
+				.PopulateSearchItems_Lambda([ConcertClient, &ReassignmentLogic](const ConcertSharedSlate::FReplicatedObjectData& ObjectData, TArray<FString>& InOutSearchStrings)
 				{
 					SReassignObjectComboBox::PopulateSearchTerms(*ConcertClient->GetCurrentSession(), ReassignmentLogic, ObjectData.GetObjectPath(), InOutSearchStrings);
 				})
diff --git a/Engine/Plugins/Developer/Concert/ConcertApp/MultiUserClient/Source/MultiUserClient/Private/Widgets/ActiveSession/Replication/Client/Multi/Columns/MultiStreamColumns_ReplicationToggle.cpp b/Engine/Plugins/Developer/Concert/ConcertApp/MultiUserClient/Source/MultiUserClient/Private/Widgets/ActiveSession/Replication/Client/Multi/Columns/MultiStreamColumns_ReplicationToggle.cpp
--- a/Engine/Plugins/Developer/Concert/ConcertApp/MultiUserClient/Source/MultiUserClient/Private/Widgets/ActiveSession/Replication/Client/Multi/Columns/MultiStreamColumns_ReplicationToggle.cpp
+++ b/Engine/Plugins/Developer/Concert/ConcertApp/MultiUserClient/Source/MultiUserClient/Private/Widgets/ActiveSession/Replication/Client/Multi/Columns/MultiStreamColumns_ReplicationToggle.cpp
@@ -14,14 +14,14 @@ namespace UE::MultiUserClient::MultiStreamColumns
 {
 	const FName ReplicationToggleColumnId(TEXT("ReplicationToggleColumn"));
 	
-	ConcertClientSharedSlate::ReplicationColumns::FReplicationTopLevelObjectColumn ReplicationToggle(
+	ConcertSharedSlate::ReplicationColumns::FReplicationTopLevelObjectColumn ReplicationToggle(
 		TSharedRef<IConcertClient> ConcertClient,
-		TAttribute<ConcertClientSharedSlate::IReplicationStreamModel*> ConsolidatedStreamModelAttribute,
+		TAttribute<ConcertSharedSlate::IReplicationStreamModel*> ConsolidatedStreamModelAttribute,
 		FReplicationClientManager& ClientManager,
 		const int32 ColumnsSortPriority
 		)
 	{
-		using namespace ConcertClientSharedSlate::ReplicationColumns;
+		using namespace ConcertSharedSlate::ReplicationColumns;
 		
 		return FReplicationTopLevelObjectColumn(
 			FReplicationTopLevelObjectColumn::FArguments()
diff --git a/Engine/Plugins/Developer/Concert/ConcertApp/MultiUserClient/Source/MultiUserClient/Private/Widgets/ActiveSession/Replication/Client/Multi/Columns/SAssignPropertyComboBox.cpp b/Engine/Plugins/Developer/Concert/ConcertApp/MultiUserClient/Source/MultiUserClient/Private/Widgets/ActiveSession/Replication/Client/Multi/Columns/SAssignPropertyComboBox.cpp
--- a/Engine/Plugins/Developer/Concert/ConcertApp/MultiUserClient/Source/MultiUserClient/Private/Widgets/ActiveSession/Replication/Client/Multi/Columns/SAssignPropertyComboBox.cpp
+++ b/Engine/Plugins/Developer/Concert/ConcertApp/MultiUserClient/Source/MultiUserClient/Private/Widgets/ActiveSession/Replication/Client/Multi/Columns/SAssignPropertyComboBox.cpp
@@ -61,7 +61,7 @@ namespace UE::MultiUserClient
 	}
 
 	void SAssignPropertyComboBox::Construct(const FArguments& InArgs,
-	    TSharedRef<ConcertClientSharedSlate::IMultiReplicationStreamEditor> InEditor,
+	    TSharedRef<ConcertSharedSlate::IMultiReplicationStreamEditor> InEditor,
 	    TSharedRef<IConcertClient> InConcertClient,
 	    FReplicationClientManager& InClientManager
 	)
@@ -186,7 +186,7 @@ namespace UE::MultiUserClient
 
 		const ECheckBoxState CheckBoxState = GetOptionCheckState(EndpointId);
 		const bool bRemoveProperty = CheckBoxState == ECheckBoxState::Checked; 
-		const TSharedRef<ConcertClientSharedSlate::IEditableReplicationStreamModel> EditModel = Client->GetClientEditModel();
+		const TSharedRef<ConcertSharedSlate::IEditableReplicationStreamModel> EditModel = Client->GetClientEditModel();
 		for (const FSoftObjectPath& ObjectPath : EditedObjects)
 		{
 			if (bRemoveProperty)
@@ -272,7 +272,7 @@ namespace UE::MultiUserClient
 			return ECheckBoxState::Unchecked;
 		}
 
-		const TSharedRef<ConcertClientSharedSlate::IEditableReplicationStreamModel> Model = Client->GetClientEditModel();
+		const TSharedRef<ConcertSharedSlate::IEditableReplicationStreamModel> Model = Client->GetClientEditModel();
 		ECheckBoxState CheckBoxState = ECheckBoxState::Undetermined;
 		for (const FSoftObjectPath& ObjectPath : EditedObjects)
 		{
diff --git a/Engine/Plugins/Developer/Concert/ConcertApp/MultiUserClient/Source/MultiUserClient/Private/Widgets/ActiveSession/Replication/Client/Multi/Columns/SAssignPropertyComboBox.h b/Engine/Plugins/Developer/Concert/ConcertApp/MultiUserClient/Source/MultiUserClient/Private/Widgets/ActiveSession/Replication/Client/Multi/Columns/SAssignPropertyComboBox.h
--- a/Engine/Plugins/Developer/Concert/ConcertApp/MultiUserClient/Source/MultiUserClient/Private/Widgets/ActiveSession/Replication/Client/Multi/Columns/SAssignPropertyComboBox.h
+++ b/Engine/Plugins/Developer/Concert/ConcertApp/MultiUserClient/Source/MultiUserClient/Private/Widgets/ActiveSession/Replication/Client/Multi/Columns/SAssignPropertyComboBox.h
@@ -14,11 +14,8 @@ enum class ECheckBoxState : uint8;
 
 struct FGuid;
 
-namespace UE::ConcertClientSharedSlate
-{
-	class SHorizontalClientList;
-	class IMultiReplicationStreamEditor;
-}
+namespace UE::ConcertClientSharedSlate { class SHorizontalClientList; }
+namespace UE::ConcertSharedSlate { class IMultiReplicationStreamEditor; }
 
 namespace UE::MultiUserClient
 {
@@ -54,7 +51,7 @@ namespace UE::MultiUserClient
 		SLATE_END_ARGS()
 
 		void Construct(const FArguments& InArgs,
-           TSharedRef<ConcertClientSharedSlate::IMultiReplicationStreamEditor> InEditor,
+           TSharedRef<ConcertSharedSlate::IMultiReplicationStreamEditor> InEditor,
            TSharedRef<IConcertClient> InConcertClient,
            FReplicationClientManager& InClientManager
 		);
@@ -64,7 +61,7 @@ namespace UE::MultiUserClient
 		FReplicationClientManager* ClientManager = nullptr;
 		
 		/** Used to obtain info about the streams */
-		TSharedPtr<ConcertClientSharedSlate::IMultiReplicationStreamEditor> Editor;
+		TSharedPtr<ConcertSharedSlate::IMultiReplicationStreamEditor> Editor;
 		TSharedPtr<IConcertClient> ConcertClient;
 
 		/** The objects for which the property is being displayed. */
diff --git a/Engine/Plugins/Developer/Concert/ConcertApp/MultiUserClient/Source/MultiUserClient/Private/Widgets/ActiveSession/Replication/Client/Multi/Columns/SReassignObjectComboBox.cpp b/Engine/Plugins/Developer/Concert/ConcertApp/MultiUserClient/Source/MultiUserClient/Private/Widgets/ActiveSession/Replication/Client/Multi/Columns/SReassignObjectComboBox.cpp
--- a/Engine/Plugins/Developer/Concert/ConcertApp/MultiUserClient/Source/MultiUserClient/Private/Widgets/ActiveSession/Replication/Client/Multi/Columns/SReassignObjectComboBox.cpp
+++ b/Engine/Plugins/Developer/Concert/ConcertApp/MultiUserClient/Source/MultiUserClient/Private/Widgets/ActiveSession/Replication/Client/Multi/Columns/SReassignObjectComboBox.cpp
@@ -204,7 +204,7 @@ namespace UE::MultiUserClient
 
 	SReassignObjectComboBox::FInlineObjectPathArray SReassignObjectComboBox::GetChildrenOfManagedObject() const
 	{
-		const ConcertClientSharedSlate::IReplicationStreamModel* Model = ConsolidatedStreamModelAttribute.Get();
+		const ConcertSharedSlate::IReplicationStreamModel* Model = ConsolidatedStreamModelAttribute.Get();
 		return ensure(Model) ? Model->GetSubobjects<FInlineAllocator>(ManagedObject) : FInlineObjectPathArray{};
 	}
 }
diff --git a/Engine/Plugins/Developer/Concert/ConcertApp/MultiUserClient/Source/MultiUserClient/Private/Widgets/ActiveSession/Replication/Client/Multi/Columns/SReassignObjectComboBox.h b/Engine/Plugins/Developer/Concert/ConcertApp/MultiUserClient/Source/MultiUserClient/Private/Widgets/ActiveSession/Replication/Client/Multi/Columns/SReassignObjectComboBox.h
--- a/Engine/Plugins/Developer/Concert/ConcertApp/MultiUserClient/Source/MultiUserClient/Private/Widgets/ActiveSession/Replication/Client/Multi/Columns/SReassignObjectComboBox.h
+++ b/Engine/Plugins/Developer/Concert/ConcertApp/MultiUserClient/Source/MultiUserClient/Private/Widgets/ActiveSession/Replication/Client/Multi/Columns/SReassignObjectComboBox.h
@@ -9,11 +9,8 @@ class FMenuBuilder;
 class IConcertClient;
 class IConcertClientSession;
 
-namespace UE::ConcertClientSharedSlate
-{
-	class IReplicationStreamModel;
-	class SHorizontalClientList;
-}
+namespace UE::ConcertClientSharedSlate { class SHorizontalClientList; }
+namespace UE::ConcertSharedSlate { class IReplicationStreamModel; }
 
 namespace UE::MultiUserClient
 {
@@ -49,7 +46,7 @@ namespace UE::MultiUserClient
 			SLATE_ARGUMENT(FSoftObjectPath, ManagedObject)
 			SLATE_ARGUMENT(TSharedPtr<FText>, HighlightText)
 			/** Used to to figure out child objects. */
-			SLATE_ATTRIBUTE(ConcertClientSharedSlate::IReplicationStreamModel*, ConsolidatedModel)
+			SLATE_ATTRIBUTE(ConcertSharedSlate::IReplicationStreamModel*, ConsolidatedModel)
 
 			/** Called when a valid client ID is selected for reassignment. */
 			SLATE_EVENT(FOnReassignAllOptionClicked, OnReassignAllOptionClicked)
@@ -76,7 +73,7 @@ namespace UE::MultiUserClient
 		/** The object being reassigned by this widget. */
 		FSoftObjectPath ManagedObject;
 		/** Used to to figure out child objects. */
-		TAttribute<ConcertClientSharedSlate::IReplicationStreamModel*> ConsolidatedStreamModelAttribute;
+		TAttribute<ConcertSharedSlate::IReplicationStreamModel*> ConsolidatedStreamModelAttribute;
 		
 		/** Displayed as content of combo button */
 		TSharedPtr<ConcertClientSharedSlate::SHorizontalClientList> ComboClientList;
diff --git a/Engine/Plugins/Developer/Concert/ConcertApp/MultiUserClient/Source/MultiUserClient/Private/Widgets/ActiveSession/Replication/Client/Multi/Columns/SReplicationMultiToggleCheckbox.h b/Engine/Plugins/Developer/Concert/ConcertApp/MultiUserClient/Source/MultiUserClient/Private/Widgets/ActiveSession/Replication/Client/Multi/Columns/SReplicationMultiToggleCheckbox.h
--- a/Engine/Plugins/Developer/Concert/ConcertApp/MultiUserClient/Source/MultiUserClient/Private/Widgets/ActiveSession/Replication/Client/Multi/Columns/SReplicationMultiToggleCheckbox.h
+++ b/Engine/Plugins/Developer/Concert/ConcertApp/MultiUserClient/Source/MultiUserClient/Private/Widgets/ActiveSession/Replication/Client/Multi/Columns/SReplicationMultiToggleCheckbox.h
@@ -9,7 +9,7 @@
 class IConcertClient;
 enum class ECheckBoxState : uint8;
 
-namespace UE::ConcertClientSharedSlate
+namespace UE::ConcertSharedSlate
 {
 	class IReplicationStreamModel;
 }
@@ -36,7 +36,7 @@ namespace UE::MultiUserClient
 		SLATE_BEGIN_ARGS(SReplicationMultiToggleCheckbox)
 		{}
 		    SLATE_ARGUMENT(FSoftObjectPath, Object)
-			SLATE_ATTRIBUTE(ConcertClientSharedSlate::IReplicationStreamModel*, ConsolidatedStreamModelAttribute)
+			SLATE_ATTRIBUTE(ConcertSharedSlate::IReplicationStreamModel*, ConsolidatedStreamModelAttribute)
 		SLATE_END_ARGS()
 
 		void Construct(
@@ -53,7 +53,7 @@ namespace UE::MultiUserClient
 		/** Used to access all clients for toggling authority. */
 		FReplicationClientManager* ClientManager = nullptr;
 		/** Used to get children of Object */
-		TAttribute<ConcertClientSharedSlate::IReplicationStreamModel*> ConsolidatedStreamModelAttribute;
+		TAttribute<ConcertSharedSlate::IReplicationStreamModel*> ConsolidatedStreamModelAttribute;
 		
 		/** Used to look up client display names in case of conflicts. */
 		TSharedPtr<IConcertClient> ConcertClient;
diff --git a/Engine/Plugins/Developer/Concert/ConcertApp/MultiUserClient/Source/MultiUserClient/Private/Widgets/ActiveSession/Replication/Client/Multi/MultiStreamModel.cpp b/Engine/Plugins/Developer/Concert/ConcertApp/MultiUserClient/Source/MultiUserClient/Private/Widgets/ActiveSession/Replication/Client/Multi/MultiStreamModel.cpp
--- a/Engine/Plugins/Developer/Concert/ConcertApp/MultiUserClient/Source/MultiUserClient/Private/Widgets/ActiveSession/Replication/Client/Multi/MultiStreamModel.cpp
+++ b/Engine/Plugins/Developer/Concert/ConcertApp/MultiUserClient/Source/MultiUserClient/Private/Widgets/ActiveSession/Replication/Client/Multi/MultiStreamModel.cpp
@@ -35,16 +35,16 @@ namespace UE::MultiUserClient
 		}
 	}
 
-	TSet<TSharedRef<ConcertClientSharedSlate::IReplicationStreamModel>> FMultiStreamModel::GetReadOnlyStreams() const
+	TSet<TSharedRef<ConcertSharedSlate::IReplicationStreamModel>> FMultiStreamModel::GetReadOnlyStreams() const
 	{
-		TSet<TSharedRef<ConcertClientSharedSlate::IReplicationStreamModel>> Result;
+		TSet<TSharedRef<ConcertSharedSlate::IReplicationStreamModel>> Result;
 		Algo::Transform(CachedReadOnlyClients, Result, [](const FReplicationClient* Client){ return Client->GetClientEditModel(); });
 		return Result;
 	}
 
-	TSet<TSharedRef<ConcertClientSharedSlate::IEditableReplicationStreamModel>> FMultiStreamModel::GetEditableStreams() const
+	TSet<TSharedRef<ConcertSharedSlate::IEditableReplicationStreamModel>> FMultiStreamModel::GetEditableStreams() const
 	{
-		TSet<TSharedRef<ConcertClientSharedSlate::IEditableReplicationStreamModel>> Result;
+		TSet<TSharedRef<ConcertSharedSlate::IEditableReplicationStreamModel>> Result;
 		Algo::Transform(CachedWritableClients, Result, [](const FReplicationClient* Client){ return Client->GetClientEditModel(); });
 		return Result;
 	}
@@ -63,7 +63,7 @@ namespace UE::MultiUserClient
 		ClientSelectionModel.ForEachSelectedClient([this, &ReadOnlyClients, &WritableClients](FReplicationClient& Client)
 		{
 			const bool bIsUploadable = CanEverSubmit(Client.GetSubmissionWorkflow().GetUploadability());
-			const TSharedRef<ConcertClientSharedSlate::IEditableReplicationStreamModel> Stream = Client.GetClientEditModel();
+			const TSharedRef<ConcertSharedSlate::IEditableReplicationStreamModel> Stream = Client.GetClientEditModel();
 			Client.OnModelChanged().AddRaw(this, &FMultiStreamModel::OnStreamExternallyChanged, Stream.ToWeakPtr());
 			
 			if (bIsUploadable)
@@ -95,9 +95,9 @@ namespace UE::MultiUserClient
 		}
 	}
 
-	void FMultiStreamModel::OnStreamExternallyChanged(TWeakPtr<ConcertClientSharedSlate::IEditableReplicationStreamModel> ChangedStream)
+	void FMultiStreamModel::OnStreamExternallyChanged(TWeakPtr<ConcertSharedSlate::IEditableReplicationStreamModel> ChangedStream)
 	{
-		if (const TSharedPtr<ConcertClientSharedSlate::IEditableReplicationStreamModel> ChangedStreamPin = ChangedStream.Pin())
+		if (const TSharedPtr<ConcertSharedSlate::IEditableReplicationStreamModel> ChangedStreamPin = ChangedStream.Pin())
 		{
 			OnReadOnlyStreamChangedDelegate.Broadcast(ChangedStreamPin.ToSharedRef());
 		}
diff --git a/Engine/Plugins/Developer/Concert/ConcertApp/MultiUserClient/Source/MultiUserClient/Private/Widgets/ActiveSession/Replication/Client/Multi/MultiStreamModel.h b/Engine/Plugins/Developer/Concert/ConcertApp/MultiUserClient/Source/MultiUserClient/Private/Widgets/ActiveSession/Replication/Client/Multi/MultiStreamModel.h
--- a/Engine/Plugins/Developer/Concert/ConcertApp/MultiUserClient/Source/MultiUserClient/Private/Widgets/ActiveSession/Replication/Client/Multi/MultiStreamModel.h
+++ b/Engine/Plugins/Developer/Concert/ConcertApp/MultiUserClient/Source/MultiUserClient/Private/Widgets/ActiveSession/Replication/Client/Multi/MultiStreamModel.h
@@ -17,7 +17,7 @@ namespace UE::MultiUserClient
 	class IClientSelectionModel;
 
 	/** Checks whether clients accept remote changes and categorizes them into read-only and writable. */
-	class FMultiStreamModel : public ConcertClientSharedSlate::IEditableMultiReplicationStreamModel
+	class FMultiStreamModel : public ConcertSharedSlate::IEditableMultiReplicationStreamModel
 	{
 	public:
 		
@@ -28,8 +28,8 @@ namespace UE::MultiUserClient
 		void ForEachClient(TFunctionRef<EBreakBehavior(const FReplicationClient*)> ProcessClient) const;
 
 		//~ Begin IEditableMultiReplicationStreamModel Interface
-		virtual TSet<TSharedRef<ConcertClientSharedSlate::IReplicationStreamModel>> GetReadOnlyStreams() const override;
-		virtual TSet<TSharedRef<ConcertClientSharedSlate::IEditableReplicationStreamModel>> GetEditableStreams() const override;
+		virtual TSet<TSharedRef<ConcertSharedSlate::IReplicationStreamModel>> GetReadOnlyStreams() const override;
+		virtual TSet<TSharedRef<ConcertSharedSlate::IEditableReplicationStreamModel>> GetEditableStreams() const override;
 		virtual FOnStreamExternallyChanged& OnStreamExternallyChanged() override { return OnReadOnlyStreamChangedDelegate; }
 		virtual FOnStreamSetChanged& OnStreamSetChanged() override { return OnStreamSetChangedDelegate; }
 		//~ End IEditableMultiReplicationStreamModel Interface
@@ -50,7 +50,7 @@ namespace UE::MultiUserClient
 		void RebuildStreamsSets();
 
 		/** Handle read-only streams changing */
-		void OnStreamExternallyChanged(TWeakPtr<ConcertClientSharedSlate::IEditableReplicationStreamModel> ChangedStream);
+		void OnStreamExternallyChanged(TWeakPtr<ConcertSharedSlate::IEditableReplicationStreamModel> ChangedStream);
 	};
 }
 
diff --git a/Engine/Plugins/Developer/Concert/ConcertApp/MultiUserClient/Source/MultiUserClient/Private/Widgets/ActiveSession/Replication/Client/Multi/SMultiClientView.cpp b/Engine/Plugins/Developer/Concert/ConcertApp/MultiUserClient/Source/MultiUserClient/Private/Widgets/ActiveSession/Replication/Client/Multi/SMultiClientView.cpp
--- a/Engine/Plugins/Developer/Concert/ConcertApp/MultiUserClient/Source/MultiUserClient/Private/Widgets/ActiveSession/Replication/Client/Multi/SMultiClientView.cpp
+++ b/Engine/Plugins/Developer/Concert/ConcertApp/MultiUserClient/Source/MultiUserClient/Private/Widgets/ActiveSession/Replication/Client/Multi/SMultiClientView.cpp
@@ -4,18 +4,19 @@
 
 #include "IClientSelectionModel.h"
 #include "MultiStreamModel.h"
+#include "Replication/ClientReplicationWidgetFactories.h"
+#include "Replication/ReplicationWidgetFactories.h"
 #include "Replication/Client/ReplicationClient.h"
 #include "Replication/Client/ReplicationClientManager.h"
-#include "Replication/Editor/Model/Object/ActorSelectionSourceModel.h"
 #include "Replication/Editor/Model/Property/SelectPropertyFromUClassModel.h"
 #include "Replication/Editor/View/IMultiReplicationStreamEditor.h"
-#include "Replication/ReplicationWidgetFactories.h"
 #include "Replication/Editor/View/IReplicationStreamEditor.h"
+#include "Replication/Editor/Model/ObjectSource/ActorSelectionSourceModel.h"
+#include "Widgets/ActiveSession/Replication/Client/FrequencyContextMenuUtils.h"
 #include "Widgets/ActiveSession/Replication/Client/Multi/Columns/MultiStreamColumns.h"
 #include "Widgets/ActiveSession/Replication/Client/SClientToolbar.h"
 
 #include "Widgets/SBoxPanel.h"
-#include "Widgets/ActiveSession/Replication/Client/FrequencyContextMenuUtils.h"
 
 #define LOCTEXT_NAMESPACE "SMultiClientView"
 
@@ -64,7 +65,7 @@ namespace UE::MultiUserClient
 
 	TSharedRef<SWidget> SMultiClientView::CreateEditorContent(const TSharedRef<IConcertClient>& InConcertClient, FReplicationClientManager& InClientManager)
 	{
-		using namespace UE::ConcertClientSharedSlate;
+		using namespace UE::ConcertSharedSlate;
 
 		TAttribute<TSharedPtr<IMultiReplicationStreamEditor>> MultiStreamEditorAttribute =
 		   TAttribute<TSharedPtr<IMultiReplicationStreamEditor>>::CreateLambda([this]()
@@ -85,12 +86,13 @@ namespace UE::MultiUserClient
 		FCreateMultiStreamEditorParams Params
 		{
 			.MultiStreamModel = StreamModel.ToSharedRef(),
-			.ObjectSource = MakeShared<FActorSelectionSourceModel>(),
+			.ConsolidatedObjectModel = ConcertClientSharedSlate::CreateTransactionalStreamModel(),
+			.ObjectSource = MakeShared<ConcertClientSharedSlate::FActorSelectionSourceModel>(),
 			.PropertySource = MakeShared<FSelectPropertyFromUClassModel>(),
 			.GetAutoAssignToStreamDelegate = MoveTemp(GetAutoAssignTargetDelegate),
 			.ViewerParams 
 			{
-				.SubobjectModel = CreateDefaultComponentHierarchySubobjectModel(), // This makes actors have children in the top view
+				.SubobjectModel = ConcertClientSharedSlate::CreateSubobjectModelForComponentHierarchy(), // This makes actors have children in the top view
 				.OnExtendObjectsContextMenu = FExtendObjectMenu::CreateSP(this, &SMultiClientView::ExtendObjectContextMenu),
 				.AdditionalObjectColumns =
 				{
diff --git a/Engine/Plugins/Developer/Concert/ConcertApp/MultiUserClient/Source/MultiUserClient/Private/Widgets/ActiveSession/Replication/Client/Multi/SMultiClientView.h b/Engine/Plugins/Developer/Concert/ConcertApp/MultiUserClient/Source/MultiUserClient/Private/Widgets/ActiveSession/Replication/Client/Multi/SMultiClientView.h
--- a/Engine/Plugins/Developer/Concert/ConcertApp/MultiUserClient/Source/MultiUserClient/Private/Widgets/ActiveSession/Replication/Client/Multi/SMultiClientView.h
+++ b/Engine/Plugins/Developer/Concert/ConcertApp/MultiUserClient/Source/MultiUserClient/Private/Widgets/ActiveSession/Replication/Client/Multi/SMultiClientView.h
@@ -8,7 +8,7 @@
 class IConcertClient;
 class FMenuBuilder;
 
-namespace UE::ConcertClientSharedSlate
+namespace UE::ConcertSharedSlate
 {
 	class IMultiReplicationStreamEditor;
 	class IEditableReplicationStreamModel;
@@ -47,7 +47,7 @@ namespace UE::MultiUserClient
 		/** Combines the clients */
 		TSharedPtr<FMultiStreamModel> StreamModel;
 		/** Displayed in the UI. */
-		TSharedPtr<ConcertClientSharedSlate::IMultiReplicationStreamEditor> StreamEditor;
+		TSharedPtr<ConcertSharedSlate::IMultiReplicationStreamEditor> StreamEditor;
 
 		/** Creates this widget's editor content */
 		TSharedRef<SWidget> CreateEditorContent(const TSharedRef<IConcertClient>& InConcertClient, FReplicationClientManager& InClientManager);
diff --git a/Engine/Plugins/Developer/Concert/ConcertApp/MultiUserClient/Source/MultiUserClient/Private/Widgets/ActiveSession/Replication/Client/Single/Columns/SingleClientColumns.h b/Engine/Plugins/Developer/Concert/ConcertApp/MultiUserClient/Source/MultiUserClient/Private/Widgets/ActiveSession/Replication/Client/Single/Columns/SingleClientColumns.h
--- a/Engine/Plugins/Developer/Concert/ConcertApp/MultiUserClient/Source/MultiUserClient/Private/Widgets/ActiveSession/Replication/Client/Single/Columns/SingleClientColumns.h
+++ b/Engine/Plugins/Developer/Concert/ConcertApp/MultiUserClient/Source/MultiUserClient/Private/Widgets/ActiveSession/Replication/Client/Single/Columns/SingleClientColumns.h
@@ -8,7 +8,7 @@
 
 class IConcertClient;
 
-namespace UE::ConcertClientSharedSlate
+namespace UE::ConcertSharedSlate
 {
 	class IReplicationStreamViewer;
 	class IReplicationStreamModel;
@@ -24,14 +24,14 @@ namespace UE::MultiUserClient
 
 namespace UE::MultiUserClient::SingleClientColumns
 {
-	/** @see UE::ConcertClientSharedSlate::ReplicationColumns::TopLevel::ETopLevelColumnOrder */
+	/** @see UE::ConcertSharedSlate::ReplicationColumns::TopLevel::ETopLevelColumnOrder */
 	enum class ETopLevelObjectColumnOrder
 	{
 		ToggleAuthority = 0,
 		ConflictWarning = 5,
 		Owner = 40
 	};
-	/** @see UE::ConcertClientSharedSlate::ReplicationColumns::Property::EReplicationPropertyColumnOrder */
+	/** @see UE::ConcertSharedSlate::ReplicationColumns::Property::EReplicationPropertyColumnOrder */
 	enum class EPropertyColumnOrder
 	{
 		ConflictWarning = 5,
@@ -51,8 +51,8 @@ namespace UE::MultiUserClient::SingleClientColumns
 	 * @param SubmissionWorkflow Used to determine whether changing authority is at all enabled
 	 * @return Column that can be placed in the table
 	 */
-	ConcertClientSharedSlate::ReplicationColumns::FReplicationTopLevelObjectColumn ToggleTopLevelAuthority(
-		ConcertClientSharedSlate::IReplicationStreamModel& ClientStreamModel,
+	ConcertSharedSlate::ReplicationColumns::FReplicationTopLevelObjectColumn ToggleTopLevelAuthority(
+		ConcertSharedSlate::IReplicationStreamModel& ClientStreamModel,
 		FAuthorityChangeTracker& ChangeTracker,
 		ISubmissionWorkflow& SubmissionWorkflow
 		);
@@ -66,7 +66,7 @@ namespace UE::MultiUserClient::SingleClientColumns
 	 * @param AuthoritySynchronizer Used to determine whether changing authority is at all enabled
 	 * @return Column that can be placed in the table
 	 */
-	ConcertClientSharedSlate::ReplicationColumns::FReplicationTopLevelObjectColumn ToggleObjectAuthority(
+	ConcertSharedSlate::ReplicationColumns::FReplicationTopLevelObjectColumn ToggleObjectAuthority(
 		FAuthorityChangeTracker& ChangeTracker,
 		ISubmissionWorkflow& SubmissionWorkflow
 		);
@@ -83,10 +83,10 @@ namespace UE::MultiUserClient::SingleClientColumns
 	 * @param InObjectModel Used to get subobjects for a top-level object
 	 * @return Column that can be placed in the table
 	 */
-	ConcertClientSharedSlate::ReplicationColumns::FReplicationTopLevelObjectColumn OwnerOfTopLevelObject(
+	ConcertSharedSlate::ReplicationColumns::FReplicationTopLevelObjectColumn OwnerOfTopLevelObject(
 		const TSharedRef<IConcertClient>& InClient,
 		FGlobalAuthorityCache& InAuthorityCache,
-		ConcertClientSharedSlate::IReplicationStreamModel& InObjectModel
+		ConcertSharedSlate::IReplicationStreamModel& InObjectModel
 		);
 
 	/**
@@ -95,7 +95,7 @@ namespace UE::MultiUserClient::SingleClientColumns
 	 * @param InAuthorityCache Used to determine which client owns the property
 	 * @return Column that can be placed in the table
 	 */
-	ConcertClientSharedSlate::ReplicationColumns::FReplicationTopLevelObjectColumn OwnerOfObject(
+	ConcertSharedSlate::ReplicationColumns::FReplicationTopLevelObjectColumn OwnerOfObject(
 		const TSharedRef<IConcertClient>& InClient,
 		FGlobalAuthorityCache& InAuthorityCache
 		);
@@ -107,10 +107,10 @@ namespace UE::MultiUserClient::SingleClientColumns
 	 * @param InViewerAttribute Used to determine which objects the property box is displaying
 	 * @return Column that can be placed in the table
 	 */
-	ConcertClientSharedSlate::ReplicationColumns::FReplicationPropertyColumn OwnerOfProperty(
+	ConcertSharedSlate::ReplicationColumns::FReplicationPropertyColumn OwnerOfProperty(
 		const TSharedRef<IConcertClient>& InClient,
 		FGlobalAuthorityCache& InAuthorityCache,
-		const TAttribute<const ConcertClientSharedSlate::IReplicationStreamViewer*>& InViewerAttribute
+		const TAttribute<const ConcertSharedSlate::IReplicationStreamViewer*>& InViewerAttribute
 		);
 
 	
@@ -128,7 +128,7 @@ namespace UE::MultiUserClient::SingleClientColumns
 	 * 
 	 * @return Column that can be placed in the table
 	 */
-	ConcertClientSharedSlate::ReplicationColumns::FReplicationTopLevelObjectColumn ConflictWarningForObject(
+	ConcertSharedSlate::ReplicationColumns::FReplicationTopLevelObjectColumn ConflictWarningForObject(
 		TSharedRef<IConcertClient> InClient,
 		FGlobalAuthorityCache& InAuthorityCache,
 		const FGuid& ClientId
@@ -144,9 +144,9 @@ namespace UE::MultiUserClient::SingleClientColumns
 	 * 
 	 * @return Column that can be placed in the table
 	 */
-	ConcertClientSharedSlate::ReplicationColumns::FReplicationPropertyColumn ConflictWarningForProperty(
+	ConcertSharedSlate::ReplicationColumns::FReplicationPropertyColumn ConflictWarningForProperty(
 		TSharedRef<IConcertClient> InClient,
-		TAttribute<const ConcertClientSharedSlate::IReplicationStreamViewer*> InViewer,
+		TAttribute<const ConcertSharedSlate::IReplicationStreamViewer*> InViewer,
 		FGlobalAuthorityCache& InAuthorityCache,
 		const FGuid& ClientId
 		);
diff --git a/Engine/Plugins/Developer/Concert/ConcertApp/MultiUserClient/Source/MultiUserClient/Private/Widgets/ActiveSession/Replication/Client/Single/Columns/SingleClientColumns_ConflictWarning.cpp b/Engine/Plugins/Developer/Concert/ConcertApp/MultiUserClient/Source/MultiUserClient/Private/Widgets/ActiveSession/Replication/Client/Single/Columns/SingleClientColumns_ConflictWarning.cpp
--- a/Engine/Plugins/Developer/Concert/ConcertApp/MultiUserClient/Source/MultiUserClient/Private/Widgets/ActiveSession/Replication/Client/Single/Columns/SingleClientColumns_ConflictWarning.cpp
+++ b/Engine/Plugins/Developer/Concert/ConcertApp/MultiUserClient/Source/MultiUserClient/Private/Widgets/ActiveSession/Replication/Client/Single/Columns/SingleClientColumns_ConflictWarning.cpp
@@ -22,13 +22,13 @@ namespace UE::MultiUserClient::SingleClientColumns
 	const FName ConflictWarningSubobjectColumnId = TEXT("ConflictWarningSubobjectColumn");
 	const FName ConflictWarningPropertyColumnId = TEXT("ConflictWarningPropertyColumn");
 	
-	ConcertClientSharedSlate::ReplicationColumns::FReplicationTopLevelObjectColumn ConflictWarningForObject(
+	ConcertSharedSlate::ReplicationColumns::FReplicationTopLevelObjectColumn ConflictWarningForObject(
 		TSharedRef<IConcertClient> InClient,
 		FGlobalAuthorityCache& InAuthorityCache,
 		const FGuid& ClientId
 		)
 	{
-		using FColumnType = ConcertClientSharedSlate::ReplicationColumns::FReplicationTopLevelObjectColumn;
+		using FColumnType = ConcertSharedSlate::ReplicationColumns::FReplicationTopLevelObjectColumn;
 		return FColumnType(
 			FColumnType::FArguments()
 				.GenerateWidgetColumn_Lambda([InClient = MoveTemp(InClient), &InAuthorityCache, ClientId](const FColumnType::FBuildArgs& InArgs) mutable
@@ -63,7 +63,7 @@ namespace UE::MultiUserClient::SingleClientColumns
 						.Visibility_Lambda(GetVisibility)
 						.ToolTipText_Lambda(GetToolTip);
 				})
-				.PopulateSearchItems_Lambda([](const ConcertClientSharedSlate::FReplicatedObjectData&, TArray<FString>&){})
+				.PopulateSearchItems_Lambda([](const ConcertSharedSlate::FReplicatedObjectData&, TArray<FString>&){})
 				.ColumnSortOrder(static_cast<int32>(ETopLevelObjectColumnOrder::ConflictWarning)),
 			SHeaderRow::Column(ConflictWarningSubobjectColumnId)
 				.DefaultLabel(FText::GetEmpty())
@@ -71,14 +71,14 @@ namespace UE::MultiUserClient::SingleClientColumns
 			);
 	}
 
-	ConcertClientSharedSlate::ReplicationColumns::FReplicationPropertyColumn ConflictWarningForProperty(
+	ConcertSharedSlate::ReplicationColumns::FReplicationPropertyColumn ConflictWarningForProperty(
 		TSharedRef<IConcertClient> InClient,
-		TAttribute<const ConcertClientSharedSlate::IReplicationStreamViewer*> Viewer,
+		TAttribute<const ConcertSharedSlate::IReplicationStreamViewer*> Viewer,
 		FGlobalAuthorityCache& InAuthorityCache,
 		const FGuid& ClientId
 		)
 	{
-		using FColumnType = ConcertClientSharedSlate::ReplicationColumns::FReplicationPropertyColumn;
+		using FColumnType = ConcertSharedSlate::ReplicationColumns::FReplicationPropertyColumn;
 		return FColumnType(
 			FColumnType::FArguments()
 				.GenerateWidgetColumn_Lambda([InClient = MoveTemp(InClient), Viewer, &InAuthorityCache, ClientId](const FColumnType::FBuildArgs& InArgs) mutable
@@ -118,7 +118,7 @@ namespace UE::MultiUserClient::SingleClientColumns
 						.Visibility_Lambda(GetVisibility)
 						.ToolTipText_Lambda(GetToolTip);
 				})
-				.PopulateSearchItems_Lambda([](const ConcertClientSharedSlate::FReplicatedPropertyData&, TArray<FString>&){})
+				.PopulateSearchItems_Lambda([](const ConcertSharedSlate::FReplicatedPropertyData&, TArray<FString>&){})
 				.ColumnSortOrder(static_cast<int32>(EPropertyColumnOrder::ConflictWarning)),
 			SHeaderRow::Column(ConflictWarningPropertyColumnId)
 				.DefaultLabel(FText::GetEmpty())
diff --git a/Engine/Plugins/Developer/Concert/ConcertApp/MultiUserClient/Source/MultiUserClient/Private/Widgets/ActiveSession/Replication/Client/Single/Columns/SingleClientColumns_Owner.cpp b/Engine/Plugins/Developer/Concert/ConcertApp/MultiUserClient/Source/MultiUserClient/Private/Widgets/ActiveSession/Replication/Client/Single/Columns/SingleClientColumns_Owner.cpp
--- a/Engine/Plugins/Developer/Concert/ConcertApp/MultiUserClient/Source/MultiUserClient/Private/Widgets/ActiveSession/Replication/Client/Single/Columns/SingleClientColumns_Owner.cpp
+++ b/Engine/Plugins/Developer/Concert/ConcertApp/MultiUserClient/Source/MultiUserClient/Private/Widgets/ActiveSession/Replication/Client/Single/Columns/SingleClientColumns_Owner.cpp
@@ -28,12 +28,12 @@ namespace UE::MultiUserClient::SingleClientColumns
 		}
 	}
 	
-	ConcertClientSharedSlate::ReplicationColumns::FReplicationTopLevelObjectColumn OwnerOfObject(
+	ConcertSharedSlate::ReplicationColumns::FReplicationTopLevelObjectColumn OwnerOfObject(
 		const TSharedRef<IConcertClient>& InClient,
 		FGlobalAuthorityCache& InAuthorityCache
 		)
 	{
-		using FColumnType = ConcertClientSharedSlate::ReplicationColumns::FReplicationTopLevelObjectColumn;
+		using FColumnType = ConcertSharedSlate::ReplicationColumns::FReplicationTopLevelObjectColumn;
 		return FColumnType(
 			FColumnType::FArguments()
 				.GenerateWidgetColumn_Lambda([InClient, &InAuthorityCache](const FColumnType::FBuildArgs& InArgs)
@@ -45,7 +45,7 @@ namespace UE::MultiUserClient::SingleClientColumns
 						})
 						.HighlightText_Lambda([HighlightText = InArgs.HighlightText](){ return *HighlightText.Get(); });
 				})
-				.PopulateSearchItems_Lambda([InClient, &InAuthorityCache](const ConcertClientSharedSlate::FReplicatedObjectData& InArgs, TArray<FString>& InOutSearchStrings)
+				.PopulateSearchItems_Lambda([InClient, &InAuthorityCache](const ConcertSharedSlate::FReplicatedObjectData& InArgs, TArray<FString>& InOutSearchStrings)
 				{
 					Private::Shared::PopulateSearchTerms(
 						InClient,
@@ -65,7 +65,7 @@ namespace UE::MultiUserClient::SingleClientColumns
 	{
 		static TArray<FGuid> GetPropertyOwners(
 			const FGlobalAuthorityCache& InAuthorityCache,
-			const ConcertClientSharedSlate::IReplicationStreamViewer& InViewer,
+			const ConcertSharedSlate::IReplicationStreamViewer& InViewer,
 			const FConcertPropertyChain& Property
 			)
 		{
@@ -83,13 +83,13 @@ namespace UE::MultiUserClient::SingleClientColumns
 		}
 	}
 
-	ConcertClientSharedSlate::ReplicationColumns::FReplicationPropertyColumn OwnerOfProperty(
+	ConcertSharedSlate::ReplicationColumns::FReplicationPropertyColumn OwnerOfProperty(
 		const TSharedRef<IConcertClient>& InClient,
 		FGlobalAuthorityCache& InAuthorityCache,
-		const TAttribute<const ConcertClientSharedSlate::IReplicationStreamViewer*>& InViewer
+		const TAttribute<const ConcertSharedSlate::IReplicationStreamViewer*>& InViewer
 		)
 	{
-		using FColumnType = ConcertClientSharedSlate::ReplicationColumns::FReplicationPropertyColumn;
+		using FColumnType = ConcertSharedSlate::ReplicationColumns::FReplicationPropertyColumn;
 		return FColumnType(
 			FColumnType::FArguments()
 				.GenerateWidgetColumn_Lambda([InClient, &InAuthorityCache, InViewer](const FColumnType::FBuildArgs& InArgs)
@@ -101,7 +101,7 @@ namespace UE::MultiUserClient::SingleClientColumns
 						})
 						.HighlightText_Lambda([HighlightText = InArgs.HighlightText](){ return *HighlightText.Get(); });
 				})
-				.PopulateSearchItems_Lambda([InClient, &InAuthorityCache, InViewer](const ConcertClientSharedSlate::FReplicatedPropertyData& InArgs, TArray<FString>& InOutSearchStrings)
+				.PopulateSearchItems_Lambda([InClient, &InAuthorityCache, InViewer](const ConcertSharedSlate::FReplicatedPropertyData& InArgs, TArray<FString>& InOutSearchStrings)
 				{
 					Private::Shared::PopulateSearchTerms(
 						InClient,
diff --git a/Engine/Plugins/Developer/Concert/ConcertApp/MultiUserClient/Source/MultiUserClient/Private/Widgets/ActiveSession/Replication/Client/Single/Columns/SingleClientColumns_ToggleAuthority.cpp b/Engine/Plugins/Developer/Concert/ConcertApp/MultiUserClient/Source/MultiUserClient/Private/Widgets/ActiveSession/Replication/Client/Single/Columns/SingleClientColumns_ToggleAuthority.cpp
--- a/Engine/Plugins/Developer/Concert/ConcertApp/MultiUserClient/Source/MultiUserClient/Private/Widgets/ActiveSession/Replication/Client/Single/Columns/SingleClientColumns_ToggleAuthority.cpp
+++ b/Engine/Plugins/Developer/Concert/ConcertApp/MultiUserClient/Source/MultiUserClient/Private/Widgets/ActiveSession/Replication/Client/Single/Columns/SingleClientColumns_ToggleAuthority.cpp
@@ -20,12 +20,12 @@ namespace UE::MultiUserClient::SingleClientColumns
 {
 	const FName ToggleTopLevelAuthorityColumnId = TEXT("ToggleTopLevelAuthorityColumn");
 
-	ConcertClientSharedSlate::ReplicationColumns::FReplicationTopLevelObjectColumn ToggleObjectAuthority(
+	ConcertSharedSlate::ReplicationColumns::FReplicationTopLevelObjectColumn ToggleObjectAuthority(
 		FAuthorityChangeTracker& ChangeTracker,
 		ISubmissionWorkflow& SubmissionWorkflow
 		)
 	{
-		using namespace ConcertClientSharedSlate;
+		using namespace ConcertSharedSlate;
 		using FColumnDelegates = TReplicationColumnDelegates<FReplicatedObjectData>;
 		return MakeCheckboxColumn<FReplicatedObjectData>(
 			ToggleTopLevelAuthorityColumnId,
diff --git a/Engine/Plugins/Developer/Concert/ConcertApp/MultiUserClient/Source/MultiUserClient/Private/Widgets/ActiveSession/Replication/Client/Single/SReplicationClientView.cpp b/Engine/Plugins/Developer/Concert/ConcertApp/MultiUserClient/Source/MultiUserClient/Private/Widgets/ActiveSession/Replication/Client/Single/SReplicationClientView.cpp
--- a/Engine/Plugins/Developer/Concert/ConcertApp/MultiUserClient/Source/MultiUserClient/Private/Widgets/ActiveSession/Replication/Client/Single/SReplicationClientView.cpp
+++ b/Engine/Plugins/Developer/Concert/ConcertApp/MultiUserClient/Source/MultiUserClient/Private/Widgets/ActiveSession/Replication/Client/Single/SReplicationClientView.cpp
@@ -2,11 +2,11 @@
 
 #include "SReplicationClientView.h"
 
+#include "Replication/ClientReplicationWidgetFactories.h"
 #include "Replication/Client/ReplicationClient.h"
-#include "Replication/Editor/Model/IEditableReplicationStreamModel.h"
-#include "Replication/Editor/Model/Object/ActorSelectionSourceModel.h"
 #include "Replication/Editor/Model/Property/SelectPropertyFromUClassModel.h"
 #include "Replication/Editor/View/IReplicationStreamEditor.h"
+#include "Replication/Editor/Model/ObjectSource/ActorSelectionSourceModel.h"
 #include "Replication/ReplicationWidgetFactories.h"
 #include "Replication/Client/ReplicationClientManager.h"
 #include "Replication/Submission/ISubmissionWorkflow.h"
@@ -60,7 +60,7 @@ namespace UE::MultiUserClient
 
 	TSharedRef<SWidget> SReplicationClientView::CreateContent(FReplicationClient& InReplicationClient)
 	{
-		using namespace ConcertClientSharedSlate;
+		using namespace ConcertSharedSlate;
 		FAuthorityChangeTracker& AuthorityTracker = InReplicationClient.GetAuthorityDiffer();
 		ISubmissionWorkflow& SubmissionWorkflow = InReplicationClient.GetSubmissionWorkflow();
 		FGlobalAuthorityCache& AuthorityCache = ClientManager->GetAuthorityCache();
@@ -72,13 +72,13 @@ namespace UE::MultiUserClient
 		const FCreateEditorParams ReplicationEditorCreationParams
 		{
 			.DataModel = InReplicationClient.GetClientEditModel(),
-			.ObjectSource = MakeShared<FActorSelectionSourceModel>(),
+			.ObjectSource = MakeShared<ConcertClientSharedSlate::FActorSelectionSourceModel>(),
 			.PropertySource = MakeShared<FSelectPropertyFromUClassModel>(),
 			.IsEditingEnabled = TAttribute<bool>::CreateLambda([&SubmissionWorkflow](){ return CanEverSubmit(SubmissionWorkflow.GetUploadability()); }),
 			.EditingDisabledToolTipText = LOCTEXT("Editing.NotImplemented", "Editing remote clients is not implemented. You can only edit the local client."),
 			.ViewerParams =
 			{
-				.SubobjectModel = CreateDefaultComponentHierarchySubobjectModel(), // This makes actors have children in the top view
+				.SubobjectModel = ConcertClientSharedSlate::CreateSubobjectModelForComponentHierarchy(), // This makes actors have children in the top view
 				.OnExtendObjectsContextMenu = FExtendObjectMenu::CreateSP(this, &SReplicationClientView::ExtendObjectContextMenu),
 				.AdditionalObjectColumns =
 				{
@@ -94,7 +94,7 @@ namespace UE::MultiUserClient
 			}
 		};
 
-		EditorView = CreateDefaultStreamEditor(ReplicationEditorCreationParams);
+		EditorView = ConcertClientSharedSlate::CreateDefaultStreamEditor(ReplicationEditorCreationParams);
 		return EditorView.ToSharedRef();
 	}
 
diff --git a/Engine/Plugins/Developer/Concert/ConcertApp/MultiUserClient/Source/MultiUserClient/Private/Widgets/ActiveSession/Replication/Client/Single/SReplicationClientView.h b/Engine/Plugins/Developer/Concert/ConcertApp/MultiUserClient/Source/MultiUserClient/Private/Widgets/ActiveSession/Replication/Client/Single/SReplicationClientView.h
--- a/Engine/Plugins/Developer/Concert/ConcertApp/MultiUserClient/Source/MultiUserClient/Private/Widgets/ActiveSession/Replication/Client/Single/SReplicationClientView.h
+++ b/Engine/Plugins/Developer/Concert/ConcertApp/MultiUserClient/Source/MultiUserClient/Private/Widgets/ActiveSession/Replication/Client/Single/SReplicationClientView.h
@@ -8,7 +8,7 @@
 class IConcertClient;
 class FMenuBuilder;
 
-namespace UE::ConcertClientSharedSlate
+namespace UE::ConcertSharedSlate
 {
 	class IReplicationStreamEditor;
 }
@@ -41,7 +41,7 @@ namespace UE::MultiUserClient
 		FReplicationClientManager* ClientManager;
 		
 		/** The editor view of the replication content. */
-		TSharedPtr<ConcertClientSharedSlate::IReplicationStreamEditor> EditorView;
+		TSharedPtr<ConcertSharedSlate::IReplicationStreamEditor> EditorView;
 		
 		/** The client to depict. Should always return true. If the client is destroyed, so should this widget be. */
 		TAttribute<FReplicationClient*> GetReplicationClientAttribute;
diff --git a/Engine/Plugins/Developer/Concert/ConcertApp/MultiUserClient/Source/MultiUserReplicationEditor/Private/StreamEditor/ReplicationSessionPresetEditorToolkit.cpp b/Engine/Plugins/Developer/Concert/ConcertApp/MultiUserClient/Source/MultiUserReplicationEditor/Private/StreamEditor/ReplicationSessionPresetEditorToolkit.cpp
--- a/Engine/Plugins/Developer/Concert/ConcertApp/MultiUserClient/Source/MultiUserReplicationEditor/Private/StreamEditor/ReplicationSessionPresetEditorToolkit.cpp
+++ b/Engine/Plugins/Developer/Concert/ConcertApp/MultiUserClient/Source/MultiUserReplicationEditor/Private/StreamEditor/ReplicationSessionPresetEditorToolkit.cpp
@@ -5,12 +5,13 @@
 #include "Assets/MultiUserReplicationClientPreset.h"
 #include "Assets/MultiUserReplicationSessionPreset.h"
 #include "Replication/Editor/Model/Property/SelectPropertyFromUClassModel.h"
-#include "Replication/Editor/Model/Object/ActorSelectionSourceModel.h"
 #include "Replication/Editor/View/IReplicationStreamEditor.h"
+#include "Replication/Editor/Model/ObjectSource/ActorSelectionSourceModel.h"
 #include "Replication/ReplicationWidgetFactories.h"
 
 #include "Framework/Docking/TabManager.h"
 #include "Framework/Docking/WorkspaceItem.h"
+#include "Replication/ClientReplicationWidgetFactories.h"
 #include "Widgets/Docking/SDockTab.h"
 
 #define LOCTEXT_NAMESPACE "FReplicationStreamEditorToolkit"
@@ -73,11 +74,13 @@ namespace UE::MultiUserReplicationEditor
 	TSharedRef<SDockTab> FReplicationSessionPresetEditorToolkit::SpawnTab_Content(const FSpawnTabArgs& SpawnTabArgs)
 	{
 		using namespace ConcertClientSharedSlate;
+		using namespace ConcertSharedSlate;
 		
-		const TSharedRef<IEditableReplicationStreamModel> AssetReadWriteModel = CreatePropertySelectionModel(
-			*GetEditedStreamAsset(),
-			GetEditedStreamAsset()->GetUnassignedClient()->Stream->MakeReplicationMapGetterAttribute()
+		const TSharedRef<IEditableReplicationStreamModel> AssetReadWriteModel = CreateTransactionalStreamModel(
+			CreateBaseStreamModel(GetEditedStreamAsset()->GetUnassignedClient()->Stream->MakeReplicationMapGetterAttribute()),
+			*GetEditedStreamAsset()
 			);
+		
 		const TSharedRef<FActorSelectionSourceModel> ObjectSourceModel = MakeShared<FActorSelectionSourceModel>();
 		const TSharedRef<FSelectPropertyFromUClassModel> PropertySourceModel = MakeShared<FSelectPropertyFromUClassModel>();
 		const TSharedRef<IReplicationStreamEditor> EditorView = CreateDefaultStreamEditor(FCreateEditorParams{ AssetReadWriteModel, ObjectSourceModel, PropertySourceModel });
diff --git a/Engine/Plugins/Developer/Concert/ConcertUI/ConcertClientSharedSlate/Source/ConcertClientSharedSlate/Private/Replication/ClientReplicationWidgetFactories.cpp b/Engine/Plugins/Developer/Concert/ConcertUI/ConcertClientSharedSlate/Source/ConcertClientSharedSlate/Private/Replication/ClientReplicationWidgetFactories.cpp
new file mode 100644
--- /dev/null
+++ b/Engine/Plugins/Developer/Concert/ConcertUI/ConcertClientSharedSlate/Source/ConcertClientSharedSlate/Private/Replication/ClientReplicationWidgetFactories.cpp
@@ -0,0 +1,53 @@
+// Copyright Epic Games, Inc. All Rights Reserved.
+
+#include "Replication/ClientReplicationWidgetFactories.h"
+
+#include "Replication/ReplicationWidgetFactories.h"
+#include "Replication/Editor/Model/ReplicationStreamObject.h"
+#include "Replication/Editor/Model/Subobject/ComponentHierarchySubobjectModel.h"
+#include "Replication/Editor/Model/TransactionalReplicationStreamModel.h"
+#include "Replication/Editor/View/ObjectEditor/SDefaultReplicationStreamEditor.h"
+
+#include "UObject/UObjectGlobals.h"
+#include "UObject/Package.h"
+
+namespace UE::ConcertClientSharedSlate
+{
+	TSharedRef<ConcertSharedSlate::ISubobjectModel> CreateSubobjectModelForComponentHierarchy()
+	{
+		return MakeShared<FComponentHierarchySubobjectModel>();
+	}
+
+	TSharedRef<ConcertSharedSlate::IEditableReplicationStreamModel> CreateTransactionalStreamModel(
+		TSharedRef<ConcertSharedSlate::IEditableReplicationStreamModel> BaseModel,
+		UObject& OwnerObject
+		)
+	{
+		return MakeShared<ConcertSharedSlate::FTransactionalReplicationStreamModel>(
+			MoveTemp(BaseModel),
+			OwnerObject
+			);
+	}
+
+	TSharedRef<ConcertSharedSlate::IEditableReplicationStreamModel> CreateTransactionalStreamModel()
+	{
+		const EObjectFlags Flags = RF_Transient | RF_Transactional;
+		UReplicationStreamObject* Object = NewObject<UReplicationStreamObject>(GetTransientPackage(), NAME_None, Flags);
+
+		TAttribute<FObjectReplicationMap*> Attribute = TAttribute<FObjectReplicationMap*>::CreateLambda([WeakPtr = TWeakObjectPtr<UReplicationStreamObject>(Object)]() -> FObjectReplicationMap* 
+		{
+			if (WeakPtr.IsValid())
+			{
+				return &WeakPtr->ReplicationMap;
+			}
+			return nullptr;
+		});
+
+		return CreateTransactionalStreamModel(ConcertSharedSlate::CreateBaseStreamModel(MoveTemp(Attribute)), *Object);
+	}
+	
+	TSharedRef<ConcertSharedSlate::IReplicationStreamEditor> CreateDefaultStreamEditor(ConcertSharedSlate::FCreateEditorParams Params)
+	{
+		return SNew(SDefaultReplicationStreamEditor, MoveTemp(Params));
+	}
+}
\ No newline at end of file
diff --git a/Engine/Plugins/Developer/Concert/ConcertUI/ConcertClientSharedSlate/Source/ConcertClientSharedSlate/Private/Replication/Editor/Model/Subobject/ComponentHierarchySubobjectModel.cpp b/Engine/Plugins/Developer/Concert/ConcertUI/ConcertClientSharedSlate/Source/ConcertClientSharedSlate/Private/Replication/Editor/Model/Subobject/ComponentHierarchySubobjectModel.cpp
--- a/Engine/Plugins/Developer/Concert/ConcertUI/ConcertClientSharedSlate/Source/ConcertClientSharedSlate/Private/Replication/Editor/Model/Subobject/ComponentHierarchySubobjectModel.cpp
+++ b/Engine/Plugins/Developer/Concert/ConcertUI/ConcertClientSharedSlate/Source/ConcertClientSharedSlate/Private/Replication/Editor/Model/Subobject/ComponentHierarchySubobjectModel.cpp
@@ -2,12 +2,12 @@
 
 #include "ComponentHierarchySubobjectModel.h"
 
-#include "SubobjectDataSubsystem.h"
 #include "Components/SceneComponent.h"
 #include "Misc/EBreakBehavior.h"
+#include "Replication/ObjectUtils.h"
 
 #include "GameFramework/Actor.h"
-#include "Replication/Editor/View/ObjectUtils.h"
+#include "SubobjectDataSubsystem.h"
 
 namespace UE::ConcertClientSharedSlate
 {
@@ -51,7 +51,7 @@ namespace UE::ConcertClientSharedSlate
 
 	bool FComponentHierarchySubobjectModel::IsTopLevelObject(const FSoftObjectPath& Object) const
 	{
-		return ObjectUtils::IsActor(Object);
+		return ConcertSharedSlate::ObjectUtils::IsActor(Object);
 	}
 
 	FText FComponentHierarchySubobjectModel::GetSubobjectDisplayName(const FSoftObjectPath& ObjectPath) const
diff --git a/Engine/Plugins/Developer/Concert/ConcertUI/ConcertClientSharedSlate/Source/ConcertClientSharedSlate/Private/Replication/Editor/Model/Subobject/ComponentHierarchySubobjectModel.h b/Engine/Plugins/Developer/Concert/ConcertUI/ConcertClientSharedSlate/Source/ConcertClientSharedSlate/Private/Replication/Editor/Model/Subobject/ComponentHierarchySubobjectModel.h
--- a/Engine/Plugins/Developer/Concert/ConcertUI/ConcertClientSharedSlate/Source/ConcertClientSharedSlate/Private/Replication/Editor/Model/Subobject/ComponentHierarchySubobjectModel.h
+++ b/Engine/Plugins/Developer/Concert/ConcertUI/ConcertClientSharedSlate/Source/ConcertClientSharedSlate/Private/Replication/Editor/Model/Subobject/ComponentHierarchySubobjectModel.h
@@ -13,7 +13,7 @@ namespace UE::ConcertClientSharedSlate
 	 * Builds a similar tree hierarchy as SSubobjectEditor.
 	 * This model only inspects UActorComponents on the root object, which must be an actor.
 	 */
-	class FComponentHierarchySubobjectModel : public ISubobjectModel
+	class FComponentHierarchySubobjectModel : public ConcertSharedSlate::ISubobjectModel
 	{
 	public:
 
diff --git a/Engine/Plugins/Developer/Concert/ConcertUI/ConcertClientSharedSlate/Source/ConcertClientSharedSlate/Private/Replication/Editor/Model/TransactionalReplicationStreamModel.cpp b/Engine/Plugins/Developer/Concert/ConcertUI/ConcertClientSharedSlate/Source/ConcertClientSharedSlate/Private/Replication/Editor/Model/TransactionalReplicationStreamModel.cpp
--- a/Engine/Plugins/Developer/Concert/ConcertUI/ConcertClientSharedSlate/Source/ConcertClientSharedSlate/Private/Replication/Editor/Model/TransactionalReplicationStreamModel.cpp
+++ b/Engine/Plugins/Developer/Concert/ConcertUI/ConcertClientSharedSlate/Source/ConcertClientSharedSlate/Private/Replication/Editor/Model/TransactionalReplicationStreamModel.cpp
@@ -7,14 +7,13 @@
 
 #define LOCTEXT_NAMESPACE "FTransactionalReplicationStreamModel"
 
-namespace UE::ConcertClientSharedSlate
+namespace UE::ConcertSharedSlate
 {
 	FTransactionalReplicationStreamModel::FTransactionalReplicationStreamModel(
-		UObject& OwningObject,
-		TAttribute<FObjectReplicationMap*> ReplicationMapAttribute,
-		TSharedPtr<IStreamExtender> Extender
+		TSharedRef<IEditableReplicationStreamModel> WrappedModel,
+		UObject& OwningObject
 		)
-		: FGenericReplicationStreamModel(MoveTemp(ReplicationMapAttribute), MoveTemp(Extender))
+		: WrappedModel(MoveTemp(WrappedModel))
 		, OwningObject(&OwningObject)
 	{}
 
@@ -22,28 +21,28 @@ namespace UE::ConcertClientSharedSlate
 	{
 		const FScopedTransaction Transaction(LOCTEXT("AddObjects", "Add replicated objects"));
 		OwningObject->Modify();
-		FGenericReplicationStreamModel::AddObjects(Objects);
+		WrappedModel->AddObjects(Objects);
 	}
 
 	void FTransactionalReplicationStreamModel::RemoveObjects(TConstArrayView<FSoftObjectPath> Objects)
 	{
 		const FScopedTransaction Transaction(LOCTEXT("RemoveObjects", "Remove replicated objects"));
 		OwningObject->Modify();
-		FGenericReplicationStreamModel::RemoveObjects(Objects);
+		WrappedModel->RemoveObjects(Objects);
 	}
 
 	void FTransactionalReplicationStreamModel::AddProperties(const FSoftObjectPath& SoftObjectPath, TConstArrayView<FConcertPropertyChain> Properties)
 	{
 		const FScopedTransaction Transaction(LOCTEXT("AddProperties", "Add replicated properties"));
 		OwningObject->Modify();
-		FGenericReplicationStreamModel::AddProperties(SoftObjectPath, Properties);
+		WrappedModel->AddProperties(SoftObjectPath, Properties);
 	}
 
 	void FTransactionalReplicationStreamModel::RemoveProperties(const FSoftObjectPath& SoftObjectPath, TConstArrayView<FConcertPropertyChain> Properties)
 	{
 		const FScopedTransaction Transaction(LOCTEXT("RemoveProperties", "Remove replicated properties"));
 		OwningObject->Modify();
-		FGenericReplicationStreamModel::RemoveProperties(SoftObjectPath, Properties);
+		WrappedModel->RemoveProperties(SoftObjectPath, Properties);
 	}
 
 	bool FTransactionalReplicationStreamModel::MatchesContext(const FTransactionContext& InContext, const TArray<TPair<UObject*, FTransactionObjectEvent>>& TransactionObjectContexts) const
@@ -80,6 +79,11 @@ namespace UE::ConcertClientSharedSlate
 		OnObjectsChanged().Broadcast({}, {}, EReplicatedObjectChangeReason::ExternalChange);
 		OnPropertiesChanged().Broadcast();
 	}
+
+	void FTransactionalReplicationStreamModel::AddReferencedObjects(FReferenceCollector& Collector)
+	{
+		Collector.AddReferencedObject(OwningObject);
+	}
 }
 
 #undef LOCTEXT_NAMESPACE
\ No newline at end of file
diff --git a/Engine/Plugins/Developer/Concert/ConcertUI/ConcertClientSharedSlate/Source/ConcertClientSharedSlate/Private/Replication/Editor/Model/TransactionalReplicationStreamModel.h b/Engine/Plugins/Developer/Concert/ConcertUI/ConcertClientSharedSlate/Source/ConcertClientSharedSlate/Private/Replication/Editor/Model/TransactionalReplicationStreamModel.h
--- a/Engine/Plugins/Developer/Concert/ConcertUI/ConcertClientSharedSlate/Source/ConcertClientSharedSlate/Private/Replication/Editor/Model/TransactionalReplicationStreamModel.h
+++ b/Engine/Plugins/Developer/Concert/ConcertUI/ConcertClientSharedSlate/Source/ConcertClientSharedSlate/Private/Replication/Editor/Model/TransactionalReplicationStreamModel.h
@@ -2,25 +2,38 @@
 
 #pragma once
 
-#include "GenericReplicationStreamModel.h"
 #include "EditorUndoClient.h"
+#include "Replication/Editor/Model/IEditableReplicationStreamModel.h"
+#include "UObject/GCObject.h"
 
-namespace UE::ConcertClientSharedSlate
+struct FObjectReplicationMap;
+
+namespace UE::ConcertSharedSlate
 {
 	class IStreamExtender;
 
 	/** Special case of FGenericPropertySelectionModel where the edited FObjectReplicationMap lives in an UObject that is RF_Transactional. */
 	class FTransactionalReplicationStreamModel
-		: public FGenericReplicationStreamModel
+		: public IEditableReplicationStreamModel
 		, public FSelfRegisteringEditorUndoClient
+		, public FGCObject
 	{
 	public:
 
 		FTransactionalReplicationStreamModel(
-			UObject& OwningObject,
-			TAttribute<FObjectReplicationMap*> ReplicationMapAttribute,
-			TSharedPtr<IStreamExtender> Extender = nullptr
+			TSharedRef<IEditableReplicationStreamModel> WrappedModel,
+			UObject& OwningObject
 			);
+
+		//~ Begin IReplicationStreamModel Interface
+		virtual FSoftClassPath GetObjectClass(const FSoftObjectPath& Object) const override { return WrappedModel->GetObjectClass(Object); }
+		virtual bool ContainsObjects(const TSet<FSoftObjectPath>& Objects) const override { return WrappedModel->ContainsObjects(Objects); }
+		virtual bool ContainsProperties(const FSoftObjectPath& Object, const TSet<FConcertPropertyChain>& Properties) const override { return WrappedModel->ContainsProperties(Object, Properties); }
+		virtual bool ForEachReplicatedObject(TFunctionRef<EBreakBehavior(const FSoftObjectPath& Object)> Delegate) const override { return WrappedModel->ForEachReplicatedObject(Delegate); }
+		virtual bool ForEachProperty(const FSoftObjectPath& Object, TFunctionRef<EBreakBehavior(const FConcertPropertyChain& Property)> Delegate) const override { return WrappedModel->ForEachProperty(Object, Delegate); }
+		virtual FOnObjectsChanged& OnObjectsChanged() override { return WrappedModel->OnObjectsChanged(); }
+		virtual FOnPropertiesChanged& OnPropertiesChanged() override { return WrappedModel->OnPropertiesChanged(); }
+		//~ End IReplicationStreamModel Interface
 		
 		//~ Begin IEditableReplicationStreamModel Interface
 		virtual void AddObjects(TConstArrayView<UObject*> Objects) override;
@@ -35,8 +48,16 @@ namespace UE::ConcertClientSharedSlate
 		virtual void PostRedo(bool bSuccess) override;
 		//~ End FEditorUndoClient Interface
 
+		//~ Begin FGCObject Interface
+		virtual void AddReferencedObjects(FReferenceCollector& Collector) override;
+		virtual FString GetReferencerName() const override { return TEXT("FTransactionalReplicationStreamModel"); }
+		//~ End FGCObject Interface
+
 	private:
 
+		/** The model to be transacted */
+		TSharedRef<IEditableReplicationStreamModel> WrappedModel;
+		
 		/** User of FTransactionalPropertySelectionModel is responsible for keeping OwningObject alive, e.g. via an asset editor. */
 		TWeakObjectPtr<UObject> OwningObject;
 	};
diff --git a/Engine/Plugins/Developer/Concert/ConcertUI/ConcertClientSharedSlate/Source/ConcertClientSharedSlate/Private/Replication/Editor/View/ClientEditorColumns.cpp b/Engine/Plugins/Developer/Concert/ConcertUI/ConcertClientSharedSlate/Source/ConcertClientSharedSlate/Private/Replication/Editor/View/ClientEditorColumns.cpp
new file mode 100644
--- /dev/null
+++ b/Engine/Plugins/Developer/Concert/ConcertUI/ConcertClientSharedSlate/Source/ConcertClientSharedSlate/Private/Replication/Editor/View/ClientEditorColumns.cpp
@@ -0,0 +1,161 @@
+// Copyright Epic Games, Inc. All Rights Reserved.
+
+#include "Replication/Editor/View/ClientEditorColumns.h"
+
+#include "Replication/Editor/Model/IEditableReplicationStreamModel.h"
+#include "Replication/Editor/Model/ReplicatedPropertyData.h"
+#include "Replication/Editor/View/ReplicationColumnsUtils.h"
+#include "Replication/ObjectUtils.h"
+
+#include "Internationalization/Internationalization.h"
+#include "Replication/Editor/View/IReplicationStreamViewer.h"
+#include "Replication/Editor/View/PredefinedReplicationColumns.h"
+
+#define LOCTEXT_NAMESPACE "ReplicationPropertyColumns"
+
+namespace UE::ConcertClientSharedSlate::ReplicationColumns::Property
+{
+	const FName ReplicatesColumnId = TEXT("ReplicatedColumn");
+	
+	namespace Private
+	{
+		static ECheckBoxState OnGetPropertyCheckboxState(
+			const FConcertPropertyChain& PropertyChain,
+			const ConcertSharedSlate::IReplicationStreamViewer& Viewer,
+			const ConcertSharedSlate::IReplicationStreamModel& Model
+			)
+		{
+			const TArray<FSoftObjectPath> SelectedObjectPaths = Viewer.GetObjectsBeingPropertyEdited();
+			return GetPropertyCheckboxStateBasedOnSelection(PropertyChain, SelectedObjectPaths, Model);
+		}
+
+		static void OnPropertyCheckboxChanged(
+			bool bIsChecked,
+			const FConcertPropertyChain& PropertyChain,
+			const ConcertSharedSlate::IReplicationStreamViewer& Viewer,
+			ConcertSharedSlate::IEditableReplicationStreamModel& Model
+			)
+		{
+			const TArray PropertyAsArray{ PropertyChain };
+			const TArray<FSoftObjectPath> SelectedObjects = Viewer.GetObjectsBeingPropertyEdited();
+
+			if (bIsChecked)
+			{
+				// We cannot proceed if any of the selected objects cannot be loaded because we cannot obtain its class
+				for (const FSoftObjectPath& Path : SelectedObjects)
+				{
+					UObject* Object = Path.ResolveObject();
+					if (!Object)
+					{
+						return;
+					}
+				}
+				
+				for (const FSoftObjectPath& Path : Viewer.GetObjectsBeingPropertyEdited())
+				{
+					UObject* Object = Path.ResolveObject();
+					
+					if (!Model.ContainsObjects({ Path }))
+					{
+						// FYI: This should always be a subobject because actors are in the view beforehand (this checkbox is created in response to clicking the actor)
+						Model.AddObjects({ Object });
+					}
+					Model.AddProperties(Path, PropertyAsArray);
+				}
+			}
+			else
+			{
+				for (const FSoftObjectPath& Path : Viewer.GetObjectsBeingPropertyEdited())
+				{
+					Model.RemoveProperties(Path, PropertyAsArray);
+
+					// The actor was already in the view before this checkbox was created (the user must click an object for this checkbox to be created)
+					// So in short, this checkbox should only remove the objects it has previously added.
+					// This checkbox should remove subobjects because they were added by this checkbox before but not actors.
+					// Also FYI: actors cause their subobjects to be automatically displayed in the top view so UX wise it would be weird if you uncheck
+					// a property checkbox in the bottom view and then in the top view the object disappears.
+					if (ConcertSharedSlate::ObjectUtils::IsActor(Path) && !Model.HasAnyPropertyAssigned(Path))
+					{
+						Model.RemoveObjects({ Path });
+					}
+				}
+			}
+		}
+	}
+	
+	ConcertSharedSlate::ReplicationColumns::FReplicationPropertyColumn ReplicatesColumns(
+		TWeakPtr<ConcertSharedSlate::IReplicationStreamViewer> Viewer,
+		TWeakPtr<ConcertSharedSlate::IEditableReplicationStreamModel> Model,
+		ConcertSharedSlate::TReplicationColumnDelegates<ConcertSharedSlate::FReplicatedPropertyData>::FIsEnabled IsEnabledDelegate,
+		TAttribute<FText> DisabledToolTipText,
+		const float ColumnWidth,
+		const int32 Priority
+		)
+	{
+		using namespace ConcertSharedSlate;
+		using FPropertyColumnDelegates = TReplicationColumnDelegates<FReplicatedPropertyData>;
+		return MakeCheckboxColumn<FReplicatedPropertyData>(
+			ReplicatesColumnId,
+			FPropertyColumnDelegates(
+				FPropertyColumnDelegates::FGetColumnCheckboxState::CreateLambda(
+				[Viewer, Model](const FReplicatedPropertyData& Data)
+				{
+					const TSharedPtr<IReplicationStreamViewer> ViewerPin = Viewer.Pin();
+					const TSharedPtr<IEditableReplicationStreamModel> ModelPin = Model.Pin();
+					return ensure(ViewerPin && ModelPin) ? Private::OnGetPropertyCheckboxState(Data.GetProperty(), *ViewerPin, *ModelPin) : ECheckBoxState::Undetermined;
+				}),
+				FPropertyColumnDelegates::FOnColumnCheckboxChanged::CreateLambda(
+				[Viewer, Model](bool bIsChecked, const FReplicatedPropertyData& Data)
+				{
+					const TSharedPtr<IReplicationStreamViewer> ViewerPin = Viewer.Pin();
+					const TSharedPtr<IEditableReplicationStreamModel> ModelPin = Model.Pin();
+					if (ensure(ViewerPin && ModelPin))
+					{
+						Private::OnPropertyCheckboxChanged(bIsChecked, Data.GetProperty(), *ViewerPin, *ModelPin);
+					}
+				}),
+				FPropertyColumnDelegates::FGetToolTipText::CreateLambda(
+				[IsEnabledDelegate, DisabledToolTipText = MoveTemp(DisabledToolTipText)](const FReplicatedPropertyData& Data)
+				{
+					const bool bIsDisabled = IsEnabledDelegate.IsBound() && !IsEnabledDelegate.Execute(Data);
+					const bool bCanCall = DisabledToolTipText.IsBound() || DisabledToolTipText.IsSet();
+					return bIsDisabled
+						? bCanCall ? DisabledToolTipText.Get() : FText::GetEmpty()
+						: LOCTEXT("Replicates.ToolTip", "Select whether this property should be replicated");
+				}),
+				IsEnabledDelegate
+				),
+			FText::GetEmpty(),
+			Priority,
+			ColumnWidth
+		);
+	}
+
+	ECheckBoxState GetPropertyCheckboxStateBasedOnSelection(
+		const FConcertPropertyChain& Property,
+		TConstArrayView<FSoftObjectPath> Selection,
+		const ConcertSharedSlate::IReplicationStreamModel& Model)
+	{
+		ECheckBoxState CheckBoxState = ECheckBoxState::Undetermined;
+		for (const FSoftObjectPath& SelectedObject : Selection)
+		{
+			const bool bContainsProperty = Model.ContainsProperties(SelectedObject, { Property });
+			const ECheckBoxState StateForThisObject = bContainsProperty
+				? ECheckBoxState::Checked
+				: ECheckBoxState::Unchecked;
+			// The first object?
+			if (CheckBoxState == ECheckBoxState::Undetermined)
+			{
+				CheckBoxState = StateForThisObject;
+			}
+			else if (CheckBoxState != StateForThisObject)
+			{
+				// Two boxes do not have the same value
+				return ECheckBoxState::Undetermined;
+			}
+		}
+		return CheckBoxState;
+	}
+}
+
+#undef LOCTEXT_NAMESPACE
\ No newline at end of file
diff --git a/Engine/Plugins/Developer/Concert/ConcertUI/ConcertClientSharedSlate/Source/ConcertClientSharedSlate/Private/Replication/Editor/View/ClientEditorColumns.h b/Engine/Plugins/Developer/Concert/ConcertUI/ConcertClientSharedSlate/Source/ConcertClientSharedSlate/Private/Replication/Editor/View/ClientEditorColumns.h
new file mode 100644
--- /dev/null
+++ b/Engine/Plugins/Developer/Concert/ConcertUI/ConcertClientSharedSlate/Source/ConcertClientSharedSlate/Private/Replication/Editor/View/ClientEditorColumns.h
@@ -0,0 +1,58 @@
+// Copyright Epic Games, Inc. All Rights Reserved.
+
+#pragma once
+
+#include "Replication/Editor/View/PredefinedReplicationColumns.h"
+#include "Replication/Editor/View/ReplicationColumnsUtils.h"
+#include "Templates/SharedPointer.h"
+
+namespace UE::ConcertSharedSlate
+{
+	class IEditableReplicationStreamModel;
+	class IReplicationStreamViewer;
+	class IReplicationStreamModel;
+	class ISubobjectModel;
+}
+
+namespace UE::ConcertClientSharedSlate::ReplicationColumns::Property
+{
+	extern const FName ReplicatesColumnId;
+	
+	enum class EReplicationPropertyColumnOrder : int32
+	{
+		/** The checkbox in SDefaultReplicationStreamEditor determining whether the property is in the selection*/
+		ReplicatesCheckbox = 0,
+	};
+	
+	/**
+	 * A checkbox that is placed at the beginning of every property.
+	 * Checking & unchecking adds & removes the property to the selected objects' property mapping, respectively.
+	 *
+	 * @param Viewer Used to determine which objects are currently selected (and being edited by the checkbox)
+	 * @param Model Used to actually edit the object to properties mapping
+	 * @param IsEnabledDelegate Determines whether the checkbox is enabled
+	 * @param DisabledToolTipText Tooltip to display when IsEnabledDelegate returns false
+	 * @param ColumnWidth Width to use for the column
+	 * @param Priority Determines position of this columns relative to the others
+	 */
+	ConcertSharedSlate::ReplicationColumns::FReplicationPropertyColumn ReplicatesColumns(
+		TWeakPtr<ConcertSharedSlate::IReplicationStreamViewer> Viewer,
+		TWeakPtr<ConcertSharedSlate::IEditableReplicationStreamModel> Model,
+		ConcertSharedSlate::TReplicationColumnDelegates<ConcertSharedSlate::FReplicatedPropertyData>::FIsEnabled IsEnabledDelegate = {},
+		TAttribute<FText> DisabledToolTipText = {},
+		const float ColumnWidth = 20.f,
+		const int32 Priority = static_cast<int32>(EReplicationPropertyColumnOrder::ReplicatesCheckbox)
+		);
+	
+	/**
+	 * Goes through all selected objects, checks whether the property is checked on it or not, and returns a checkbox state.
+	 * The state is checked or unchecked if all selected properties have the property enabled or disabled, respectively, and undetermined otherwise, i.e. if it is mixed.
+	 *
+	 * This is used by ReplicatesColumns.
+	 */
+	ECheckBoxState GetPropertyCheckboxStateBasedOnSelection(
+		const FConcertPropertyChain& Property,
+		TConstArrayView<FSoftObjectPath> Selection,
+		const ConcertSharedSlate::IReplicationStreamModel& Model
+		);
+}
diff --git a/Engine/Plugins/Developer/Concert/ConcertUI/ConcertClientSharedSlate/Source/ConcertClientSharedSlate/Private/Replication/Editor/View/ObjectEditor/SDefaultReplicationStreamEditor.cpp b/Engine/Plugins/Developer/Concert/ConcertUI/ConcertClientSharedSlate/Source/ConcertClientSharedSlate/Private/Replication/Editor/View/ObjectEditor/SDefaultReplicationStreamEditor.cpp
--- a/Engine/Plugins/Developer/Concert/ConcertUI/ConcertClientSharedSlate/Source/ConcertClientSharedSlate/Private/Replication/Editor/View/ObjectEditor/SDefaultReplicationStreamEditor.cpp
+++ b/Engine/Plugins/Developer/Concert/ConcertUI/ConcertClientSharedSlate/Source/ConcertClientSharedSlate/Private/Replication/Editor/View/ObjectEditor/SDefaultReplicationStreamEditor.cpp
@@ -2,64 +2,45 @@
 
 #include "SDefaultReplicationStreamEditor.h"
 
-#include "Replication/Editor/Model/IEditableReplicationStreamModel.h"
-#include "Replication/Editor/View/Tree/SelectionViewerColumns.h"
-
-#include "SBaseReplicationStreamEditor.h"
+#include "Replication/ReplicationWidgetFactories.h"
+#include "Replication/Editor/View/ClientEditorColumns.h"
+#include "Replication/Editor/View/SelectionViewerColumns.h"
 
 namespace UE::ConcertClientSharedSlate
 {
-	void SDefaultReplicationStreamEditor::Construct(
-		const FArguments& InArgs,
-		TSharedRef<IEditableReplicationStreamModel> InPropertiesModel,
-		TSharedRef<IObjectSelectionSourceModel> InObjectSelectionSource,
-		TSharedRef<IPropertySelectionSourceModel> InPropertySelectionSource)
+	void SDefaultReplicationStreamEditor::Construct(const FArguments& InArgs, ConcertSharedSlate::FCreateEditorParams EditorParams)
 	{
-		PropertiesModel = InPropertiesModel;
+		using namespace ConcertSharedSlate;
+		using namespace ConcertSharedSlate::ReplicationColumns;
+		using namespace ConcertSharedSlate::ReplicationColumns::Property;
+		using namespace ConcertClientSharedSlate::ReplicationColumns::Property;
 		
-		using namespace ReplicationColumns;
-		using namespace ReplicationColumns::Property;
+		PropertiesModel = EditorParams.DataModel;
+
 		const FReplicationPropertyColumn ReplicatesColumn = ReplicatesColumns(
 			SharedThis(this),
-			InPropertiesModel,
-			TReplicationColumnDelegates<FReplicatedPropertyData>::FIsEnabled::CreateLambda([IsEnabled = InArgs._IsEditingEnabled](const FReplicatedPropertyData&)
+			PropertiesModel,
+			TReplicationColumnDelegates<FReplicatedPropertyData>::FIsEnabled::CreateLambda([IsEnabled = EditorParams.IsEditingEnabled](const FReplicatedPropertyData&)
 			{
 				return !IsEnabled.IsBound() || IsEnabled.Get();
 			}),
-			InArgs._EditingDisabledToolTipText
+			EditorParams.EditingDisabledToolTipText
 			);
-		TArray<FReplicationPropertyColumn> PropertyColumns = InArgs._AdditionalPropertyColumns;
+		TArray<FReplicationPropertyColumn>& PropertyColumns = EditorParams.ViewerParams.AdditionalPropertyColumns;
 		PropertyColumns.Add(ReplicatesColumn);
 
 		// Set both primary and secondary in case one is overriden but always use the override.
-		const FColumnSortInfo PrimaryPropertySort = InArgs._PrimaryObjectSort.IsValid()
-			? InArgs._PrimaryObjectSort
+		EditorParams.ViewerParams.PrimaryPropertySort = EditorParams.ViewerParams.PrimaryPropertySort.IsValid()
+			? EditorParams.ViewerParams.PrimaryPropertySort
 			: FColumnSortInfo{ ReplicatesColumnId, EColumnSortMode::Ascending };
-		const FColumnSortInfo SecondaryPropertySort = InArgs._SecondaryObjectSort.IsValid()
-			? InArgs._SecondaryObjectSort
+		EditorParams.ViewerParams.SecondaryPropertySort = EditorParams.ViewerParams.SecondaryPropertySort.IsValid()
+			? EditorParams.ViewerParams.SecondaryPropertySort
 			: FColumnSortInfo{ LabelColumnId, EColumnSortMode::Ascending };
-		
+
+		WrappedEditor = CreateBaseStreamEditor(MoveTemp(EditorParams));
 		ChildSlot
 		[
-			SAssignNew(WrappedEditor, SBaseReplicationStreamEditor, InPropertiesModel, InObjectSelectionSource, InPropertySelectionSource)
-				.AdditionalObjectColumns(InArgs._AdditionalObjectColumns)
-				.PrimaryObjectSort(InArgs._PrimaryObjectSort)
-				.SecondaryObjectSort(InArgs._SecondaryObjectSort)
-				.AdditionalPropertyColumns(PropertyColumns)
-				.PrimaryPropertySort(PrimaryPropertySort)
-				.SecondaryPropertySort(SecondaryPropertySort)
-				.SubobjectModel(InArgs._SubobjectModel)
-				.OnExtendObjectsContextMenu(InArgs._OnExtendObjectsContextMenu)
-				.LeftOfObjectSearchBar()
-				[
-					InArgs._LeftOfObjectSearchBar.Widget
-				]
-				.LeftOfPropertySearchBar()
-				[
-					InArgs._LeftOfPropertySearchBar.Widget
-				]
-				.IsEditingEnabled(InArgs._IsEditingEnabled)
-				.EditingDisabledToolTipText(InArgs._EditingDisabledToolTipText)
+			WrappedEditor.ToSharedRef()
 		];
 	}
 
