Move client independent replication UI code to ConcertSharedSlate. This to future proof the code in case we want replication UI e.g. on the server.

diff --git a/Engine/Plugins/Developer/Concert/ConcertUI/ConcertClientSharedSlate/Source/ConcertClientSharedSlate/Private/Replication/Editor/View/ObjectEditor/SDefaultReplicationStreamEditor.h b/Engine/Plugins/Developer/Concert/ConcertUI/ConcertClientSharedSlate/Source/ConcertClientSharedSlate/Private/Replication/Editor/View/ObjectEditor/SDefaultReplicationStreamEditor.h
--- a/Engine/Plugins/Developer/Concert/ConcertUI/ConcertClientSharedSlate/Source/ConcertClientSharedSlate/Private/Replication/Editor/View/ObjectEditor/SDefaultReplicationStreamEditor.h
+++ b/Engine/Plugins/Developer/Concert/ConcertUI/ConcertClientSharedSlate/Source/ConcertClientSharedSlate/Private/Replication/Editor/View/ObjectEditor/SDefaultReplicationStreamEditor.h
@@ -6,7 +6,6 @@
 #include "Replication/Editor/View/PredefinedReplicationColumns.h"
 #include "Replication/Editor/View/ReplicationColumnsUtils.h"
 #include "Replication/ReplicationWidgetDelegates.h"
-#include "Replication/Editor/View/Tree/SelectionViewerColumns.h"
 
 #include "Widgets/DeclarativeSyntaxSupport.h"
 
@@ -14,69 +13,28 @@ class SHorizontalBox;
 struct FConcertPropertyChain;
 struct FConcertStreamObjectAutoBindingRules;
 
-namespace UE::ConcertClientSharedSlate
+namespace UE::ConcertSharedSlate
 {
+	struct FCreateEditorParams;
 	class IEditableReplicationStreamModel;
-	class IObjectSelectionSourceModel;
-	class IPropertySelectionSourceModel;
-	class IReplicationSubobjectView;
-	class ISubobjectModel;
-	class SBaseReplicationStreamEditor;
+}
 
+namespace UE::ConcertClientSharedSlate
+{
 	/**
 	 * The default replication stream editor adds a checkbox in front of every property row.
 	 * 
 	 * The checkbox allows adding / removing the properties to / from the selected objects.
 	 * The properties are sorted based on the check state: checked come before unchecked.
 	 */
-	class SDefaultReplicationStreamEditor : public IReplicationStreamEditor
+	class SDefaultReplicationStreamEditor : public ConcertSharedSlate::IReplicationStreamEditor
     {
     public:
 		
-		SLATE_BEGIN_ARGS(SDefaultReplicationStreamEditor)
-		{}
-			/** Additional columns to add to the object view */
-			SLATE_ARGUMENT(TArray<ReplicationColumns::FReplicationTopLevelObjectColumn>, AdditionalObjectColumns)
-			/** Initial primary sort to set. */
-			SLATE_ARGUMENT(FColumnSortInfo, PrimaryObjectSort)
-			/** Initial secondary sort to set. */
-			SLATE_ARGUMENT(FColumnSortInfo, SecondaryObjectSort)
-			
-			/** Additional columns to add to the property view */
-			SLATE_ARGUMENT(TArray<ReplicationColumns::FReplicationPropertyColumn>, AdditionalPropertyColumns)
-			/** Initial primary sort to set. */
-			SLATE_ARGUMENT(FColumnSortInfo, PrimaryPropertySort)
-			/** Initial secondary sort to set. */
-			SLATE_ARGUMENT(FColumnSortInfo, SecondaryPropertySort)
-			
-			/** Optional. Placed between root object outliner and property editor. */
-			SLATE_ARGUMENT(TSharedPtr<IReplicationSubobjectView>, SubobjectView)
-			/** Optional. If set, this determines the children nested under the root objects. */
-			SLATE_ARGUMENT(TSharedPtr<ISubobjectModel>, SubobjectModel)
-
-			/** Called to generate the context menu for objects. This extends the options already generated by this widget. */
-			SLATE_EVENT(FExtendObjectMenu, OnExtendObjectsContextMenu)
-
-			/** Optional. Used for determining the order in which properties are displayed. */
-			SLATE_EVENT(FSortPropertyPredicate, SortPropertyRowPredicate)
-
-			/** Optional widget to add to the left of the object list search bar. */
-			SLATE_NAMED_SLOT(FArguments, LeftOfObjectSearchBar)
-			/** Optional widget to add to the left of the property list search bar. */
-			SLATE_NAMED_SLOT(FArguments, LeftOfPropertySearchBar)
-
-			/** Optional. Determines whether all UI for changing the model should be disabled. */
-			SLATE_ATTRIBUTE(bool, IsEditingEnabled)
-			/** Optional. Whenever IsEditingEnabled returns true, this tooltip is displayed for relevant, disabled UI. */
-			SLATE_ATTRIBUTE(FText, EditingDisabledToolTipText)
+		SLATE_BEGIN_ARGS(SDefaultReplicationStreamEditor) {}
 		SLATE_END_ARGS()
 
-		void Construct(
-			const FArguments& InArgs,
-			TSharedRef<IEditableReplicationStreamModel> InPropertiesModel,
-			TSharedRef<IObjectSelectionSourceModel> InObjectSelectionSource,
-			TSharedRef<IPropertySelectionSourceModel> InPropertySelectionSource
-			);
+		void Construct(const FArguments& InArgs, ConcertSharedSlate::FCreateEditorParams EditorParams);
 
 		//~ Begin IReplicationStreamEditor Interface
 		virtual void Refresh() override;
@@ -88,10 +46,10 @@ namespace UE::ConcertClientSharedSlate
 	private:
 
 		/** Core editor implementing generic editing UI & workflows. */
-		TSharedPtr<SBaseReplicationStreamEditor> WrappedEditor;
+		TSharedPtr<IReplicationStreamEditor> WrappedEditor;
 		
 		/** Model needed for sorting properties based on selection. */
-		TSharedPtr<IEditableReplicationStreamModel> PropertiesModel;
+		TSharedPtr<ConcertSharedSlate::IEditableReplicationStreamModel> PropertiesModel;
     };
 }
 
diff --git a/Engine/Plugins/Developer/Concert/ConcertUI/ConcertClientSharedSlate/Source/ConcertClientSharedSlate/Private/Replication/Settings/ConcertPerClassSubobjectMatchingRules.cpp b/Engine/Plugins/Developer/Concert/ConcertUI/ConcertClientSharedSlate/Source/ConcertClientSharedSlate/Private/Replication/Settings/ConcertPerClassSubobjectMatchingRules.cpp
--- a/Engine/Plugins/Developer/Concert/ConcertUI/ConcertClientSharedSlate/Source/ConcertClientSharedSlate/Private/Replication/Settings/ConcertPerClassSubobjectMatchingRules.cpp
+++ b/Engine/Plugins/Developer/Concert/ConcertUI/ConcertClientSharedSlate/Source/ConcertClientSharedSlate/Private/Replication/Settings/ConcertPerClassSubobjectMatchingRules.cpp
@@ -4,7 +4,7 @@
 
 #include "UObject/Class.h"
 
-namespace UE::ConcertClientSharedSlate::DefaultSubobjects
+namespace UE::ConcertSharedSlate::DefaultSubobjects
 {
 	static void InternalAddAdditionalObjectsFromSettings(UClass& StartClass, const FConcertPerClassSubobjectMatchingRules& Settings, const UObject& AddedObject, TFunctionRef<EBreakBehavior(UObject&)> OnSubobjectMatched)
 	{
@@ -43,5 +43,5 @@ namespace UE::ConcertClientSharedSlate::DefaultSubobjects
 void FConcertPerClassSubobjectMatchingRules::MatchSubobjectsRecursivelyBreakable(const UObject& Object, TFunctionRef<EBreakBehavior(UObject&)> OnSubobjectMatched) const
 {
 	checkf(!Object.IsA<UClass>(), TEXT("Pass in the UObject instanced directly, not its class!"));
-	UE::ConcertClientSharedSlate::DefaultSubobjects::InternalAddAdditionalObjectsFromSettings(*Object.GetClass(), *this, Object, OnSubobjectMatched);
+	UE::ConcertSharedSlate::DefaultSubobjects::InternalAddAdditionalObjectsFromSettings(*Object.GetClass(), *this, Object, OnSubobjectMatched);
 }
\ No newline at end of file
diff --git a/Engine/Plugins/Developer/Concert/ConcertUI/ConcertClientSharedSlate/Source/ConcertClientSharedSlate/Private/Replication/Settings/ConcertStreamObjectAutoBindingRules.cpp b/Engine/Plugins/Developer/Concert/ConcertUI/ConcertClientSharedSlate/Source/ConcertClientSharedSlate/Private/Replication/Settings/ConcertStreamObjectAutoBindingRules.cpp
--- a/Engine/Plugins/Developer/Concert/ConcertUI/ConcertClientSharedSlate/Source/ConcertClientSharedSlate/Private/Replication/Settings/ConcertStreamObjectAutoBindingRules.cpp
+++ b/Engine/Plugins/Developer/Concert/ConcertUI/ConcertClientSharedSlate/Source/ConcertClientSharedSlate/Private/Replication/Settings/ConcertStreamObjectAutoBindingRules.cpp
@@ -8,7 +8,7 @@
 #include "Algo/IndexOf.h"
 #include "UObject/Class.h"
 
-namespace UE::ConcertClientSharedSlate::DefaultProperties
+namespace UE::ConcertSharedSlate::DefaultProperties
 {
 	static void ApplyDefaultPropertySelection(TFunctionRef<void(FConcertPropertyChain&& Chain)> Callback, const FConcertDefaultPropertySelection& Selection, UStruct& Class)
 	{
@@ -63,7 +63,7 @@ void FConcertStreamObjectAutoBindingRules::AddDefaultPropertiesFromSettings(UCla
 			continue;
 		}
 			
-		UE::ConcertClientSharedSlate::DefaultProperties::ApplyDefaultPropertySelection(Callback, *DefaultProperties, *Current);
+		UE::ConcertSharedSlate::DefaultProperties::ApplyDefaultPropertySelection(Callback, *DefaultProperties, *Current);
 		// Recurse super structs
 		if (UClass* Parent = Current->GetSuperClass()
 			; Parent && DefaultProperties->bInheritFromBase)
diff --git a/Engine/Plugins/Developer/Concert/ConcertUI/ConcertClientSharedSlate/Source/ConcertClientSharedSlate/Private/Replication/Settings/ConcertSubobjectMatchingRules.cpp b/Engine/Plugins/Developer/Concert/ConcertUI/ConcertClientSharedSlate/Source/ConcertClientSharedSlate/Private/Replication/Settings/ConcertSubobjectMatchingRules.cpp
--- a/Engine/Plugins/Developer/Concert/ConcertUI/ConcertClientSharedSlate/Source/ConcertClientSharedSlate/Private/Replication/Settings/ConcertSubobjectMatchingRules.cpp
+++ b/Engine/Plugins/Developer/Concert/ConcertUI/ConcertClientSharedSlate/Source/ConcertClientSharedSlate/Private/Replication/Settings/ConcertSubobjectMatchingRules.cpp
@@ -6,7 +6,7 @@
 #include "Internationalization/Regex.h"
 #include "Misc/EBreakBehavior.h"
 
-namespace UE::ConcertClientSharedSlate
+namespace UE::ConcertSharedSlate
 {
 	static bool MatchesAnyRegex(const FString& Input, const TSet<FString>& AllRegex)
 	{
@@ -31,7 +31,7 @@ void FConcertSubobjectMatchingRules::MatchToSubobjectsBreakable(const UObject& A
 	ForEachObjectWithOuterBreakable(&AddedObject, [this, &OnSubobjectMatched](UObject* Subobject)
 	{
 		// Has exclusion regex?
-		const bool bShouldExclude = UE::ConcertClientSharedSlate::MatchesAnyRegex(Subobject->GetName(), ExcludeSubobjectRegex);
+		const bool bShouldExclude = UE::ConcertSharedSlate::MatchesAnyRegex(Subobject->GetName(), ExcludeSubobjectRegex);
 		if (bShouldExclude)
 		{
 			return true;
@@ -57,7 +57,7 @@ void FConcertSubobjectMatchingRules::MatchToSubobjectsBreakable(const UObject& A
 		}
 
 		// Has inclusion regex?
-		const bool bShouldIncludeByRegex = UE::ConcertClientSharedSlate::MatchesAnyRegex(Subobject->GetName(), IncludeSubobjectRegex);
+		const bool bShouldIncludeByRegex = UE::ConcertSharedSlate::MatchesAnyRegex(Subobject->GetName(), IncludeSubobjectRegex);
 		if (bShouldIncludeByRegex)
 		{
 			return OnSubobjectMatched(*Subobject) == EBreakBehavior::Break;
diff --git a/Engine/Plugins/Developer/Concert/ConcertUI/ConcertClientSharedSlate/Source/ConcertClientSharedSlate/Public/Replication/ClientReplicationWidgetFactories.h b/Engine/Plugins/Developer/Concert/ConcertUI/ConcertClientSharedSlate/Source/ConcertClientSharedSlate/Public/Replication/ClientReplicationWidgetFactories.h
new file mode 100644
--- /dev/null
+++ b/Engine/Plugins/Developer/Concert/ConcertUI/ConcertClientSharedSlate/Source/ConcertClientSharedSlate/Public/Replication/ClientReplicationWidgetFactories.h
@@ -0,0 +1,48 @@
+// Copyright Epic Games, Inc. All Rights Reserved.
+
+#pragma once
+
+#include "Templates/SharedPointer.h"
+
+struct FConcertStreamObjectAutoBindingRules;
+struct FObjectReplicationMap;
+
+namespace UE::ConcertSharedSlate
+{
+	class IEditableReplicationStreamModel;
+	class IReplicationStreamEditor;
+	class IStreamExtender;
+	class ISubobjectModel;
+	struct FCreateEditorParams;
+}
+
+namespace UE::ConcertClientSharedSlate
+{
+	/** Builds a similar tree hierarchy as SSubobjectEditor. Reports only components as subobjects. */
+	CONCERTCLIENTSHAREDSLATE_API TSharedRef<ConcertSharedSlate::ISubobjectModel> CreateSubobjectModelForComponentHierarchy();
+	
+	/**
+	 * Wraps the passed in BaseModel and makes it transactional.
+	 * All calls that modify the underlying model was wrapped with scoped transactions.
+	 * 
+	 * @param OwnerObject The object containing the FObjectReplicationMap - used for transactions.
+	 */
+	CONCERTCLIENTSHAREDSLATE_API TSharedRef<ConcertSharedSlate::IEditableReplicationStreamModel> CreateTransactionalStreamModel(
+		TSharedRef<ConcertSharedSlate::IEditableReplicationStreamModel> BaseModel,
+		UObject& OwnerObject
+		);
+
+	/** Simpler CreateTransactionalStreamModel overload that internally creates an UObject and sets it up automatically. */
+	CONCERTCLIENTSHAREDSLATE_API TSharedRef<ConcertSharedSlate::IEditableReplicationStreamModel> CreateTransactionalStreamModel();
+	
+	/**
+	 * Creates a default IReplicationStreamEditor.
+	 *
+	 * This editor adds a checkbox to the start of every property row.
+	 * - Checking adds the property to the selected objects' property mappings.
+	 * - Unchecking removes the property to the selected objects' property mappings
+	 * 
+	 * @see CreateViewer for a description of the UI layout.
+	 */
+	CONCERTCLIENTSHAREDSLATE_API TSharedRef<ConcertSharedSlate::IReplicationStreamEditor> CreateDefaultStreamEditor(ConcertSharedSlate::FCreateEditorParams Params);
+}
diff --git a/Engine/Plugins/Developer/Concert/ConcertUI/ConcertSharedSlate/Source/ConcertSharedSlate/ConcertSharedSlate.Build.cs b/Engine/Plugins/Developer/Concert/ConcertUI/ConcertSharedSlate/Source/ConcertSharedSlate/ConcertSharedSlate.Build.cs
--- a/Engine/Plugins/Developer/Concert/ConcertUI/ConcertSharedSlate/Source/ConcertSharedSlate/ConcertSharedSlate.Build.cs
+++ b/Engine/Plugins/Developer/Concert/ConcertUI/ConcertSharedSlate/Source/ConcertSharedSlate/ConcertSharedSlate.Build.cs
@@ -28,10 +28,20 @@ namespace UnrealBuildTool.Rules
 					"Json",
 					"Projects",
 					"SlateCore",
+					"ToolMenus",
 					"ToolWidgets"
 				}
 			);
 
+			if (Target.bBuildEditor)
+			{
+				PrivateDependencyModuleNames.AddRange(
+					new string[]
+					{
+						"Engine" // In editor builds, we want to display actors using their label. See DisplayUtils::GetObjectDisplayString
+					}
+				);
+			}
 		}
 	}
 }
diff --git a/Engine/Plugins/Developer/Concert/ConcertUI/ConcertSharedSlate/Source/ConcertSharedSlate/Private/Replication/Editor/Model/Property/SelectPropertyFromUClassModel.cpp b/Engine/Plugins/Developer/Concert/ConcertUI/ConcertSharedSlate/Source/ConcertSharedSlate/Private/Replication/Editor/Model/Property/SelectPropertyFromUClassModel.cpp
new file mode 100644
--- /dev/null
+++ b/Engine/Plugins/Developer/Concert/ConcertUI/ConcertSharedSlate/Source/ConcertSharedSlate/Private/Replication/Editor/Model/Property/SelectPropertyFromUClassModel.cpp
@@ -0,0 +1,32 @@
+// Copyright Epic Games, Inc. All Rights Reserved.
+
+#include "Replication/Editor/Model/Property/SelectPropertyFromUClassModel.h"
+
+#include "ConcertLogGlobal.h"
+#include "Replication/Editor/Model/Property/ConcertSyncCoreReplicatedPropertySource.h"
+
+#define LOCTEXT_NAMESPACE "FSelectPropertyFromUClassModel"
+
+namespace UE::ConcertSharedSlate
+{
+	FSelectPropertyFromUClassModel::FSelectPropertyFromUClassModel()
+		: UClassIteratorSource(MakeShared<FConcertSyncCoreReplicatedPropertySource>())
+	{}
+
+	TSharedRef<IPropertySourceModel> FSelectPropertyFromUClassModel::GetPropertySource(const FSoftClassPath& Class) const
+	{
+		UClass* LoadedClass = Class.TryLoadClass<UObject>();
+		if (LoadedClass)
+		{
+			UClassIteratorSource->SetClass(LoadedClass);
+		}
+		else
+		{
+			UClassIteratorSource->SetClass(nullptr);
+			UE_LOG(LogConcert, Warning, TEXT("Could not resolve class %s. Properties will not be available."), *Class.ToString());
+		}
+		return UClassIteratorSource;
+	}
+}
+
+#undef LOCTEXT_NAMESPACE
\ No newline at end of file
diff --git a/Engine/Plugins/Developer/Concert/ConcertUI/ConcertSharedSlate/Source/ConcertSharedSlate/Private/Replication/Editor/View/MultiEditor/SMultiReplicationStreamEditor.cpp b/Engine/Plugins/Developer/Concert/ConcertUI/ConcertSharedSlate/Source/ConcertSharedSlate/Private/Replication/Editor/View/MultiEditor/SMultiReplicationStreamEditor.cpp
new file mode 100644
--- /dev/null
+++ b/Engine/Plugins/Developer/Concert/ConcertUI/ConcertSharedSlate/Source/ConcertSharedSlate/Private/Replication/Editor/View/MultiEditor/SMultiReplicationStreamEditor.cpp
@@ -0,0 +1,31 @@
+// Copyright Epic Games, Inc. All Rights Reserved.
+
+#include "SMultiReplicationStreamEditor.h"
+
+#include "ConsolidatedMultiStreamModel.h"
+#include "Replication/ReplicationWidgetFactories.h"
+#include "Replication/Editor/Model/IEditableMultiReplicationStreamModel.h"
+#include "Replication/Editor/View/IReplicationStreamEditor.h"
+
+namespace UE::ConcertSharedSlate
+{
+	void SMultiReplicationStreamEditor::Construct(const FArguments& InArgs, FCreateMultiStreamEditorParams Params)
+	{
+		MultiStreamModel = MoveTemp(Params.MultiStreamModel);
+		ConsolidatedModel = MakeShared<FConsolidatedMultiStreamModel>(MoveTemp(Params.ConsolidatedObjectModel), MoveTemp(Params.MultiStreamModel), InArgs._GetAutoAssignStream);
+		
+		const FCreateEditorParams BaseEditorParams
+		{
+			.DataModel = ConsolidatedModel.ToSharedRef(),
+			.ObjectSource = MoveTemp(Params.ObjectSource),
+			.PropertySource = MoveTemp(Params.PropertySource),
+			.ViewerParams = MoveTemp(Params.ViewerParams)
+		};
+		EditorView = CreateBaseStreamEditor(BaseEditorParams);
+
+		ChildSlot
+		[
+			EditorView.ToSharedRef()
+		];
+	}
+}
diff --git a/Engine/Plugins/Developer/Concert/ConcertUI/ConcertSharedSlate/Source/ConcertSharedSlate/Private/Replication/Editor/View/ObjectEditor/FakeObjectToPropertiesEditorModel.cpp b/Engine/Plugins/Developer/Concert/ConcertUI/ConcertSharedSlate/Source/ConcertSharedSlate/Private/Replication/Editor/View/ObjectEditor/FakeObjectToPropertiesEditorModel.cpp
new file mode 100644
--- /dev/null
+++ b/Engine/Plugins/Developer/Concert/ConcertUI/ConcertSharedSlate/Source/ConcertSharedSlate/Private/Replication/Editor/View/ObjectEditor/FakeObjectToPropertiesEditorModel.cpp
@@ -0,0 +1,104 @@
+// Copyright Epic Games, Inc. All Rights Reserved.
+
+#include "FakeObjectToPropertiesEditorModel.h"
+
+#include "Replication/Editor/Model/Property/IPropertySelectionSourceModel.h"
+#include "Replication/Editor/Model/Property/IPropertySourceModel.h"
+#include "Replication/ObjectUtils.h"
+#include "Replication/PropertyChainUtils.h"
+
+namespace UE::ConcertSharedSlate
+{
+	FSoftClassPath FFakeObjectToPropertiesEditorModel::GetObjectClass(const FSoftObjectPath& Object) const
+	{
+		const FSoftClassPath ResolvedClass = RealModel->GetObjectClass(Object);
+		if (ResolvedClass.IsValid())
+		{
+			return ResolvedClass;
+		}
+
+#if WITH_EDITOR
+		// The object is not yet in the model.
+		const UObject* LoadedObject = Object.ResolveObject();
+		return LoadedObject ? LoadedObject->GetClass() : FSoftClassPath{};
+#else
+		return FSoftClassPath{};
+#endif
+	}
+
+	bool FFakeObjectToPropertiesEditorModel::ForEachReplicatedObject(TFunctionRef<EBreakBehavior(const FSoftObjectPath& Object)> Delegate) const
+	{
+		return IterateObjects(Delegate);
+	}
+
+	bool FFakeObjectToPropertiesEditorModel::ForEachProperty(const FSoftObjectPath& ObjectPath, TFunctionRef<EBreakBehavior(const FConcertPropertyChain& Property)> Delegate) const
+	{
+		IterateDisplayedProperties(ObjectPath, Delegate);
+		return true;
+	}
+
+	bool FFakeObjectToPropertiesEditorModel::IterateObjects(TFunctionRef<EBreakBehavior(const FSoftObjectPath& Object)> Delegate) const
+	{
+		// Case: RealModel contains only components but not the owning actor.
+		// In that case, we want the UI to still show the owning actor.
+		// We'll track this with these containers:
+		TSet<FSoftObjectPath> Actors;
+		TSet<FSoftObjectPath> NonActors;
+
+		// Here we only exposes objects that are top level so we skip components, etc.
+		EBreakBehavior BreakBehavior = EBreakBehavior::Continue;
+		const bool bResult = RealModel->ForEachReplicatedObject([this, &Delegate, &Actors, &NonActors, &BreakBehavior](const FSoftObjectPath& Object)
+		{
+			if (ObjectUtils::IsActor(Object))
+			{
+				NonActors.Add(Object);
+				BreakBehavior = Delegate(Object); 
+				return BreakBehavior;
+			}
+			
+			Actors.Add(Object);
+			return EBreakBehavior::Continue;
+		});
+
+		// No additional objects to show
+		if (BreakBehavior == EBreakBehavior::Break)
+		{
+			return bResult;
+		}
+
+		// Now determine the top level objects that are not in RealModel but that need to be shown for subobjects that are in RealModel
+		for (const FSoftObjectPath& NonTopLevelObject : Actors)
+		{
+			const TOptional<FSoftObjectPath> TopLevelObject = ObjectUtils::GetActorOf(NonTopLevelObject);
+			if (!TopLevelObject)
+			{
+				continue;
+			}
+			
+			if (!NonActors.Contains(*TopLevelObject))
+			{
+				if (Delegate(*TopLevelObject) == EBreakBehavior::Break)
+				{
+					return true;
+				}
+
+				// Avoid calling Delegate again for TopLevelObject
+				NonActors.Add(*TopLevelObject);
+			}
+		}
+
+		return bResult;
+	}
+
+	void FFakeObjectToPropertiesEditorModel::IterateDisplayedProperties(const FSoftObjectPath& ObjectPath, TFunctionRef<EBreakBehavior(const FConcertPropertyChain& Property)> Delegate) const
+	{
+		// GetObjectClass might return null because ObjectPath might not be contained in the model
+		const FSoftClassPath ClassPath = GetObjectClass(ObjectPath);
+		// This makes SObjectToPropertyView list all available properties
+		PropertySelectionSource->GetPropertySource(ClassPath)
+		   ->EnumerateSelectableItems([&Delegate](const FSelectablePropertyInfo& PropertyInfo)
+		   {
+			   return Delegate(PropertyInfo.Property);
+		   });
+	}
+}
diff --git a/Engine/Plugins/Developer/Concert/ConcertUI/ConcertSharedSlate/Source/ConcertSharedSlate/Private/Replication/ReplicationWidgetFactories.cpp b/Engine/Plugins/Developer/Concert/ConcertUI/ConcertSharedSlate/Source/ConcertSharedSlate/Private/Replication/ReplicationWidgetFactories.cpp
new file mode 100644
--- /dev/null
+++ b/Engine/Plugins/Developer/Concert/ConcertUI/ConcertSharedSlate/Source/ConcertSharedSlate/Private/Replication/ReplicationWidgetFactories.cpp
@@ -0,0 +1,42 @@
+// Copyright Epic Games, Inc. All Rights Reserved.
+
+#include "Replication/ReplicationWidgetFactories.h"
+
+#include "Editor/Model/GenericReplicationStreamModel.h"
+#include "Editor/View/MultiEditor/SMultiReplicationStreamEditor.h"
+#include "Editor/View/ObjectEditor/SBaseReplicationStreamEditor.h"
+
+namespace UE::ConcertSharedSlate
+{
+	TSharedRef<IEditableReplicationStreamModel> CreateBaseStreamModel(
+		TAttribute<FObjectReplicationMap*> ReplicationMapAttribute,
+		TSharedPtr<IStreamExtender> Extender
+		)
+	{
+		return MakeShared<FGenericReplicationStreamModel>(MoveTemp(ReplicationMapAttribute), MoveTemp(Extender));
+	}
+	
+	TSharedRef<IReplicationStreamEditor> CreateBaseStreamEditor(FCreateEditorParams Params)
+	{
+		return SNew(SBaseReplicationStreamEditor, Params.DataModel, Params.ObjectSource, Params.PropertySource)
+			.AdditionalObjectColumns(Params.ViewerParams.AdditionalObjectColumns)
+			.PrimaryObjectSort(Params.ViewerParams.PrimaryObjectSort)
+			.SecondaryObjectSort(Params.ViewerParams.SecondaryObjectSort)
+			.AdditionalPropertyColumns(Params.ViewerParams.AdditionalPropertyColumns)
+			.PrimaryPropertySort(Params.ViewerParams.PrimaryPropertySort)
+			.SecondaryPropertySort(Params.ViewerParams.SecondaryPropertySort)
+			.SubobjectModel(Params.ViewerParams.SubobjectModel)
+			.OnExtendObjectsContextMenu(Params.ViewerParams.OnExtendObjectsContextMenu)
+			.LeftOfObjectSearchBar() [ Params.ViewerParams.LeftOfObjectSearchBar.Widget ]
+			.LeftOfPropertySearchBar() [ Params.ViewerParams.LeftOfPropertySearchBar.Widget ]
+			.IsEditingEnabled(Params.IsEditingEnabled)
+			.EditingDisabledToolTipText(Params.EditingDisabledToolTipText);
+	}
+
+	TSharedRef<IMultiReplicationStreamEditor> CreateBaseMultiStreamEditor(FCreateMultiStreamEditorParams Params)
+	{
+		return SNew(SMultiReplicationStreamEditor, MoveTemp(Params))
+			.GetAutoAssignStream(MoveTemp(Params.GetAutoAssignToStreamDelegate));
+	}
+}
+
diff --git a/Engine/Plugins/Developer/Concert/ConcertUI/ConcertSharedSlate/Source/ConcertSharedSlate/Public/Replication/Editor/View/SelectionViewerColumns.h b/Engine/Plugins/Developer/Concert/ConcertUI/ConcertSharedSlate/Source/ConcertSharedSlate/Public/Replication/Editor/View/SelectionViewerColumns.h
new file mode 100644
--- /dev/null
+++ b/Engine/Plugins/Developer/Concert/ConcertUI/ConcertSharedSlate/Source/ConcertSharedSlate/Public/Replication/Editor/View/SelectionViewerColumns.h
@@ -0,0 +1,50 @@
+// Copyright Epic Games, Inc. All Rights Reserved.
+
+#pragma once
+
+#include "Replication/Editor/View/PredefinedReplicationColumns.h"
+#include "Replication/Editor/View/ReplicationColumnsUtils.h"
+#include "Templates/SharedPointer.h"
+
+namespace UE::ConcertSharedSlate
+{
+	class IEditableReplicationStreamModel;
+	class IReplicationStreamViewer;
+	class IReplicationStreamModel;
+	class ISubobjectModel;
+}
+
+namespace UE::ConcertSharedSlate::ReplicationColumns::TopLevel
+{
+	CONCERTSHAREDSLATE_API extern const FName IconColumnId;
+	CONCERTSHAREDSLATE_API extern const FName LabelColumnId;
+	CONCERTSHAREDSLATE_API extern const FName TypeColumnId;
+
+	enum class ETopLevelColumnOrder : int32
+	{
+		/** Label of the object */
+		Label = 20,
+		/** Class of the object */
+		Type = 30,
+	};
+
+	CONCERTSHAREDSLATE_API FReplicationTopLevelObjectColumn LabelColumn(TSharedRef<IReplicationStreamModel> Model, ISubobjectModel* SubobjectModel = nullptr);
+	CONCERTSHAREDSLATE_API FReplicationTopLevelObjectColumn TypeColumn(TSharedRef<IReplicationStreamModel> Model);
+}
+
+namespace UE::ConcertSharedSlate::ReplicationColumns::Property
+{
+	CONCERTSHAREDSLATE_API extern const FName LabelColumnId;
+	CONCERTSHAREDSLATE_API extern const FName TypeColumnId;
+	
+	enum class EReplicationPropertyColumnOrder : int32
+	{
+		/** Label of the property */
+		Label = 10,
+		/** Type of the property */
+		Type = 20
+	};
+	
+	CONCERTSHAREDSLATE_API FReplicationPropertyColumn LabelColumn();
+	CONCERTSHAREDSLATE_API FReplicationPropertyColumn TypeColumn();
+}
