Move property model to ConcertClientSharedSlate since it relies on editor info

diff --git a/Engine/Plugins/Developer/Concert/ConcertApp/MultiUserClient/Source/MultiUserClient/Private/Widgets/ActiveSession/Replication/Client/Multi/SMultiClientView.cpp b/Engine/Plugins/Developer/Concert/ConcertApp/MultiUserClient/Source/MultiUserClient/Private/Widgets/ActiveSession/Replication/Client/Multi/SMultiClientView.cpp
--- a/Engine/Plugins/Developer/Concert/ConcertApp/MultiUserClient/Source/MultiUserClient/Private/Widgets/ActiveSession/Replication/Client/Multi/SMultiClientView.cpp
+++ b/Engine/Plugins/Developer/Concert/ConcertApp/MultiUserClient/Source/MultiUserClient/Private/Widgets/ActiveSession/Replication/Client/Multi/SMultiClientView.cpp
@@ -8,7 +8,7 @@
 #include "Replication/ReplicationWidgetFactories.h"
 #include "Replication/Client/ReplicationClient.h"
 #include "Replication/Client/ReplicationClientManager.h"
-#include "Replication/Editor/Model/Property/SelectPropertyFromUClassModel.h"
+#include "Replication/Editor/Model/PropertySource/SelectPropertyFromUClassModel.h"
 #include "Replication/Editor/View/IMultiReplicationStreamEditor.h"
 #include "Replication/Editor/View/IReplicationStreamEditor.h"
 #include "Replication/Editor/Model/ObjectSource/ActorSelectionSourceModel.h"
@@ -88,7 +88,7 @@ namespace UE::MultiUserClient
 			.MultiStreamModel = StreamModel.ToSharedRef(),
 			.ConsolidatedObjectModel = ConcertClientSharedSlate::CreateTransactionalStreamModel(),
 			.ObjectSource = MakeShared<ConcertClientSharedSlate::FActorSelectionSourceModel>(),
-			.PropertySource = MakeShared<FSelectPropertyFromUClassModel>(),
+			.PropertySource = MakeShared<ConcertClientSharedSlate::FSelectPropertyFromUClassModel>(),
 			.GetAutoAssignToStreamDelegate = MoveTemp(GetAutoAssignTargetDelegate),
 			.ViewerParams 
 			{
diff --git a/Engine/Plugins/Developer/Concert/ConcertApp/MultiUserClient/Source/MultiUserClient/Private/Widgets/ActiveSession/Replication/Client/Single/SReplicationClientView.cpp b/Engine/Plugins/Developer/Concert/ConcertApp/MultiUserClient/Source/MultiUserClient/Private/Widgets/ActiveSession/Replication/Client/Single/SReplicationClientView.cpp
--- a/Engine/Plugins/Developer/Concert/ConcertApp/MultiUserClient/Source/MultiUserClient/Private/Widgets/ActiveSession/Replication/Client/Single/SReplicationClientView.cpp
+++ b/Engine/Plugins/Developer/Concert/ConcertApp/MultiUserClient/Source/MultiUserClient/Private/Widgets/ActiveSession/Replication/Client/Single/SReplicationClientView.cpp
@@ -4,7 +4,7 @@
 
 #include "Replication/ClientReplicationWidgetFactories.h"
 #include "Replication/Client/ReplicationClient.h"
-#include "Replication/Editor/Model/Property/SelectPropertyFromUClassModel.h"
+#include "Replication/Editor/Model/PropertySource/SelectPropertyFromUClassModel.h"
 #include "Replication/Editor/View/IReplicationStreamEditor.h"
 #include "Replication/Editor/Model/ObjectSource/ActorSelectionSourceModel.h"
 #include "Replication/ReplicationWidgetFactories.h"
@@ -73,7 +73,7 @@ namespace UE::MultiUserClient
 		{
 			.DataModel = InReplicationClient.GetClientEditModel(),
 			.ObjectSource = MakeShared<ConcertClientSharedSlate::FActorSelectionSourceModel>(),
-			.PropertySource = MakeShared<FSelectPropertyFromUClassModel>(),
+			.PropertySource = MakeShared<ConcertClientSharedSlate::FSelectPropertyFromUClassModel>(),
 			.IsEditingEnabled = TAttribute<bool>::CreateLambda([&SubmissionWorkflow](){ return CanEverSubmit(SubmissionWorkflow.GetUploadability()); }),
 			.EditingDisabledToolTipText = LOCTEXT("Editing.NotImplemented", "Editing remote clients is not implemented. You can only edit the local client."),
 			.ViewerParams =
diff --git a/Engine/Plugins/Developer/Concert/ConcertApp/MultiUserClient/Source/MultiUserReplicationEditor/Private/StreamEditor/ReplicationSessionPresetEditorToolkit.cpp b/Engine/Plugins/Developer/Concert/ConcertApp/MultiUserClient/Source/MultiUserReplicationEditor/Private/StreamEditor/ReplicationSessionPresetEditorToolkit.cpp
--- a/Engine/Plugins/Developer/Concert/ConcertApp/MultiUserClient/Source/MultiUserReplicationEditor/Private/StreamEditor/ReplicationSessionPresetEditorToolkit.cpp
+++ b/Engine/Plugins/Developer/Concert/ConcertApp/MultiUserClient/Source/MultiUserReplicationEditor/Private/StreamEditor/ReplicationSessionPresetEditorToolkit.cpp
@@ -4,7 +4,7 @@
 
 #include "Assets/MultiUserReplicationClientPreset.h"
 #include "Assets/MultiUserReplicationSessionPreset.h"
-#include "Replication/Editor/Model/Property/SelectPropertyFromUClassModel.h"
+#include "Replication/Editor/Model/PropertySource/SelectPropertyFromUClassModel.h"
 #include "Replication/Editor/View/IReplicationStreamEditor.h"
 #include "Replication/Editor/Model/ObjectSource/ActorSelectionSourceModel.h"
 #include "Replication/ReplicationWidgetFactories.h"
diff --git a/Engine/Plugins/Developer/Concert/ConcertUI/ConcertClientSharedSlate/Source/ConcertClientSharedSlate/Private/Replication/Editor/Model/PropertySource/ConcertSyncCoreReplicatedPropertySource.cpp b/Engine/Plugins/Developer/Concert/ConcertUI/ConcertClientSharedSlate/Source/ConcertClientSharedSlate/Private/Replication/Editor/Model/PropertySource/ConcertSyncCoreReplicatedPropertySource.cpp
new file mode 100644
--- /dev/null
+++ b/Engine/Plugins/Developer/Concert/ConcertUI/ConcertClientSharedSlate/Source/ConcertClientSharedSlate/Private/Replication/Editor/Model/PropertySource/ConcertSyncCoreReplicatedPropertySource.cpp
@@ -0,0 +1,55 @@
+// Copyright Epic Games, Inc. All Rights Reserved.
+
+#include "Replication/Editor/Model/PropertySource/ConcertSyncCoreReplicatedPropertySource.h"
+
+#include "ConcertLogGlobal.h"
+#include "Replication/PropertyChainUtils.h"
+#include "UObject/UnrealType.h"
+
+#define LOCTEXT_NAMESPACE "FConcertSyncCoreReplicatedPropertySource"
+
+namespace UE::ConcertClientSharedSlate
+{
+	void FConcertSyncCoreReplicatedPropertySource::SetClass(UClass* InClass)
+	{
+		Class = InClass;
+		NumProperties = 0;
+		if (!Class.IsValid())
+		{
+			return;
+		}
+		
+		ConcertSyncCore::PropertyChain::ForEachReplicatableProperty(*Class.Get(), [this](const auto&, const auto&)
+		{
+			++NumProperties;
+			return EBreakBehavior::Continue;
+		});
+	}
+
+	ConcertSharedSlate::FSourceDisplayInfo FConcertSyncCoreReplicatedPropertySource::GetDisplayInfo() const
+	{
+		return {
+			{
+				LOCTEXT("Label", "Add Property"),
+			   FText::GetEmpty(),
+			   FSlateIcon()
+			},
+			ConcertSharedSlate::ESourceType::ShowAsList
+		};
+	}
+
+	void FConcertSyncCoreReplicatedPropertySource::EnumerateSelectableItems(TFunctionRef<EBreakBehavior(const ConcertSharedSlate::FSelectablePropertyInfo& SelectableOption)> Delegate) const
+	{
+		if (!Class.IsValid())
+		{
+			return;
+		}
+
+		ConcertSyncCore::PropertyChain::ForEachReplicatableConcertProperty(*Class.Get(), [this, &Delegate](FConcertPropertyChain&& Property)
+		{
+			return Delegate({ MoveTemp(Property) });
+		});
+	}
+}
+
+#undef LOCTEXT_NAMESPACE
\ No newline at end of file
diff --git a/Engine/Plugins/Developer/Concert/ConcertUI/ConcertClientSharedSlate/Source/ConcertClientSharedSlate/Private/Replication/Editor/Model/PropertySource/SelectPropertyFromUClassModel.cpp b/Engine/Plugins/Developer/Concert/ConcertUI/ConcertClientSharedSlate/Source/ConcertClientSharedSlate/Private/Replication/Editor/Model/PropertySource/SelectPropertyFromUClassModel.cpp
new file mode 100644
--- /dev/null
+++ b/Engine/Plugins/Developer/Concert/ConcertUI/ConcertClientSharedSlate/Source/ConcertClientSharedSlate/Private/Replication/Editor/Model/PropertySource/SelectPropertyFromUClassModel.cpp
@@ -0,0 +1,32 @@
+// Copyright Epic Games, Inc. All Rights Reserved.
+
+#include "Replication/Editor/Model/PropertySource/SelectPropertyFromUClassModel.h"
+
+#include "ConcertLogGlobal.h"
+#include "Replication/Editor/Model/PropertySource/ConcertSyncCoreReplicatedPropertySource.h"
+
+#define LOCTEXT_NAMESPACE "FSelectPropertyFromUClassModel"
+
+namespace UE::ConcertClientSharedSlate
+{
+	FSelectPropertyFromUClassModel::FSelectPropertyFromUClassModel()
+		: UClassIteratorSource(MakeShared<FConcertSyncCoreReplicatedPropertySource>())
+	{}
+
+	TSharedRef<ConcertSharedSlate::IPropertySourceModel> FSelectPropertyFromUClassModel::GetPropertySource(const FSoftClassPath& Class) const
+	{
+		UClass* LoadedClass = Class.TryLoadClass<UObject>();
+		if (LoadedClass)
+		{
+			UClassIteratorSource->SetClass(LoadedClass);
+		}
+		else
+		{
+			UClassIteratorSource->SetClass(nullptr);
+			UE_LOG(LogConcert, Warning, TEXT("Could not resolve class %s. Properties will not be available."), *Class.ToString());
+		}
+		return UClassIteratorSource;
+	}
+}
+
+#undef LOCTEXT_NAMESPACE
\ No newline at end of file
diff --git a/Engine/Plugins/Developer/Concert/ConcertUI/ConcertClientSharedSlate/Source/ConcertClientSharedSlate/Public/Replication/Editor/Model/PropertySource/ConcertSyncCoreReplicatedPropertySource.h b/Engine/Plugins/Developer/Concert/ConcertUI/ConcertClientSharedSlate/Source/ConcertClientSharedSlate/Public/Replication/Editor/Model/PropertySource/ConcertSyncCoreReplicatedPropertySource.h
new file mode 100644
--- /dev/null
+++ b/Engine/Plugins/Developer/Concert/ConcertUI/ConcertClientSharedSlate/Source/ConcertClientSharedSlate/Public/Replication/Editor/Model/PropertySource/ConcertSyncCoreReplicatedPropertySource.h
@@ -0,0 +1,35 @@
+// Copyright Epic Games, Inc. All Rights Reserved.
+
+#pragma once
+
+#include "Misc/EBreakBehavior.h"
+#include "Replication/Editor/Model/Property/IPropertySourceModel.h"
+
+#include "HAL/Platform.h"
+#include "UObject/WeakObjectPtr.h"
+
+class UClass;
+
+namespace UE::ConcertClientSharedSlate
+{
+	/** The allowed properties are those returned by UE::ConcertSyncCore::ForEachReplicatableProperty.*/
+	class CONCERTCLIENTSHAREDSLATE_API FConcertSyncCoreReplicatedPropertySource : public ConcertSharedSlate::IPropertySourceModel
+	{
+	public:
+
+		void SetClass(UClass* InClass);
+		
+		//~ Begin IPropertySourceModel Interface
+		virtual ConcertSharedSlate::FSourceDisplayInfo GetDisplayInfo() const override;
+		virtual uint32 GetNumSelectableItems() const override { return NumProperties; }
+		virtual void EnumerateSelectableItems(TFunctionRef<EBreakBehavior(const ConcertSharedSlate::FSelectablePropertyInfo& SelectableOption)> Delegate) const override;
+		//~ End IPropertySourceModel Interface
+
+	private:
+
+		/** Weak ptr because it could be a Blueprint class that can be destroyed at any time. */
+		TWeakObjectPtr<UClass> Class;
+
+		int32 NumProperties = 0;
+	};
+}
diff --git a/Engine/Plugins/Developer/Concert/ConcertUI/ConcertClientSharedSlate/Source/ConcertClientSharedSlate/Public/Replication/Editor/Model/PropertySource/SelectPropertyFromUClassModel.h b/Engine/Plugins/Developer/Concert/ConcertUI/ConcertClientSharedSlate/Source/ConcertClientSharedSlate/Public/Replication/Editor/Model/PropertySource/SelectPropertyFromUClassModel.h
new file mode 100644
--- /dev/null
+++ b/Engine/Plugins/Developer/Concert/ConcertUI/ConcertClientSharedSlate/Source/ConcertClientSharedSlate/Public/Replication/Editor/Model/PropertySource/SelectPropertyFromUClassModel.h
@@ -0,0 +1,34 @@
+// Copyright Epic Games, Inc. All Rights Reserved.
+
+#pragma once
+
+#include "Replication/Editor/Model/Property/IPropertySelectionSourceModel.h"
+
+#include "Templates/SharedPointer.h"
+
+struct FSoftClassPath;
+
+namespace UE::ConcertClientSharedSlate
+{
+	class FConcertSyncCoreReplicatedPropertySource;
+	
+	/**
+	 * Decides which properties can be added to IEditableReplicationStreamModel.
+	 * The allowed properties are those returned by UE::ConcertSyncCore::ForEachReplicatableProperty.
+	 */
+	class CONCERTCLIENTSHAREDSLATE_API FSelectPropertyFromUClassModel : public ConcertSharedSlate::IPropertySelectionSourceModel
+	{
+	public:
+
+		FSelectPropertyFromUClassModel();
+
+		//~ Begin IPropertySelectionSourceModel Interface
+		virtual TSharedRef<ConcertSharedSlate::IPropertySourceModel> GetPropertySource(const FSoftClassPath& Class) const override;
+		//~ End IPropertySelectionSourceModel Interface
+
+	private:
+
+		TSharedRef<FConcertSyncCoreReplicatedPropertySource> UClassIteratorSource;
+	};
+}
+
diff --git a/Engine/Plugins/Developer/Concert/ConcertUI/ConcertSharedSlate/Source/ConcertSharedSlate/Private/Replication/Editor/Model/Property/ConcertSyncCoreReplicatedPropertySource.cpp b/Engine/Plugins/Developer/Concert/ConcertUI/ConcertSharedSlate/Source/ConcertSharedSlate/Private/Replication/Editor/Model/Property/ConcertSyncCoreReplicatedPropertySource.cpp
deleted file mode 100644
diff --git a/Engine/Plugins/Developer/Concert/ConcertUI/ConcertSharedSlate/Source/ConcertSharedSlate/Private/Replication/Editor/Model/Property/SelectPropertyFromUClassModel.cpp b/Engine/Plugins/Developer/Concert/ConcertUI/ConcertSharedSlate/Source/ConcertSharedSlate/Private/Replication/Editor/Model/Property/SelectPropertyFromUClassModel.cpp
deleted file mode 100644
diff --git a/Engine/Plugins/Developer/Concert/ConcertUI/ConcertSharedSlate/Source/ConcertSharedSlate/Public/Replication/Editor/Model/Property/ConcertSyncCoreReplicatedPropertySource.h b/Engine/Plugins/Developer/Concert/ConcertUI/ConcertSharedSlate/Source/ConcertSharedSlate/Public/Replication/Editor/Model/Property/ConcertSyncCoreReplicatedPropertySource.h
deleted file mode 100644
diff --git a/Engine/Plugins/Developer/Concert/ConcertUI/ConcertSharedSlate/Source/ConcertSharedSlate/Public/Replication/Editor/Model/Property/SelectPropertyFromUClassModel.h b/Engine/Plugins/Developer/Concert/ConcertUI/ConcertSharedSlate/Source/ConcertSharedSlate/Public/Replication/Editor/Model/Property/SelectPropertyFromUClassModel.h
deleted file mode 100644
