Submit file that was forgotten to be checked in. Fixes unit tests failing.

diff --git a/Engine/Plugins/Developer/Concert/ConcertUI/ConcertClientSharedSlate/Source/ConcertClientSharedSlate/Private/Replication/Settings/ConcertStreamObjectAutoBindingRules.cpp b/Engine/Plugins/Developer/Concert/ConcertUI/ConcertClientSharedSlate/Source/ConcertClientSharedSlate/Private/Replication/Settings/ConcertStreamObjectAutoBindingRules.cpp
--- a/Engine/Plugins/Developer/Concert/ConcertUI/ConcertClientSharedSlate/Source/ConcertClientSharedSlate/Private/Replication/Settings/ConcertStreamObjectAutoBindingRules.cpp
+++ b/Engine/Plugins/Developer/Concert/ConcertUI/ConcertClientSharedSlate/Source/ConcertClientSharedSlate/Private/Replication/Settings/ConcertStreamObjectAutoBindingRules.cpp
@@ -35,7 +35,7 @@ namespace UE::ConcertSharedSlate::DefaultProperties
 				Callback(FConcertPropertyChain(&Chain, LeafProperty));
 				Paths.RemoveAtSwap(IndexOfMatches);
 			}
-			return true;
+			return bMatches;
 		});
 
 		// Warn the user of unmatched properties so they can de-clutter their settings of useless entries
