Fix crash due to scene component reporting null children

diff --git a/Engine/Plugins/Developer/Concert/ConcertUI/ConcertClientSharedSlate/Source/ConcertClientSharedSlate/Private/Replication/Editor/Model/Object/EditorObjectHierarchyModel.cpp b/Engine/Plugins/Developer/Concert/ConcertUI/ConcertClientSharedSlate/Source/ConcertClientSharedSlate/Private/Replication/Editor/Model/Object/EditorObjectHierarchyModel.cpp
--- a/Engine/Plugins/Developer/Concert/ConcertUI/ConcertClientSharedSlate/Source/ConcertClientSharedSlate/Private/Replication/Editor/Model/Object/EditorObjectHierarchyModel.cpp
+++ b/Engine/Plugins/Developer/Concert/ConcertUI/ConcertClientSharedSlate/Source/ConcertClientSharedSlate/Private/Replication/Editor/Model/Object/EditorObjectHierarchyModel.cpp
@@ -29,14 +29,15 @@ namespace UE::ConcertClientSharedSlate
 			{
 				// Only the root component should be listed. It's children should be returned through recursive ForEachDirectChild calls.
 				USceneComponent* RootComponent = AsActor->GetRootComponent();
-				if (RootComponent && Private::IsValidComponent(*RootComponent) && Callback(RootComponent, EChildRelationship::Component) == EBreakBehavior::Break)
+				if (IsValid(RootComponent) && IsValidComponent(*RootComponent) && Callback(RootComponent, EChildRelationship::Component) == EBreakBehavior::Break)
 				{
 					return;
 				}
 
 				for (const UActorComponent* Component : TInlineComponentArray<UActorComponent*>(AsActor))
 				{
-					if (IsValidComponent(*Component)
+					if (IsValid(Component)
+						&& IsValidComponent(*Component)
 						// The only component to be found here should be the root component, the others are found via recursive ForEachDirectChild calls
 						&& !Component->IsA<USceneComponent>()
 						&& Callback(Component, EChildRelationship::Component) == EBreakBehavior::Break)
@@ -51,7 +52,8 @@ namespace UE::ConcertClientSharedSlate
 				for (int32 i = 0; i < SceneComponent->GetNumChildrenComponents(); ++i)
 				{
 					const USceneComponent* ChildComponent = SceneComponent->GetChildComponent(i);
-					if (IsValidComponent(*ChildComponent)
+					if (IsValid(ChildComponent) // GetChildComponent can return nullptr sometimes...
+						&& IsValidComponent(*ChildComponent)
 						&& Callback(ChildComponent, EChildRelationship::Component) == EBreakBehavior::Break)
 					{
 						break;
