Expose property tree view as IPropertyTreeView to allow reuse.
- Extracted filtering to ConcertSharedClientSlate because it relies on information available only during editor builds
- This is in preparation of a detail customization for FConcertPropertyChainWrapper (needed for UE-203742, which is greenlit for 5.4)

diff --git a/Engine/Plugins/Developer/Concert/ConcertApp/MultiUserClient/Source/MultiUserClient/Private/Replication/Client/ReplicationClient.cpp b/Engine/Plugins/Developer/Concert/ConcertApp/MultiUserClient/Source/MultiUserClient/Private/Replication/Client/ReplicationClient.cpp
--- a/Engine/Plugins/Developer/Concert/ConcertApp/MultiUserClient/Source/MultiUserClient/Private/Replication/Client/ReplicationClient.cpp
+++ b/Engine/Plugins/Developer/Concert/ConcertApp/MultiUserClient/Source/MultiUserClient/Private/Replication/Client/ReplicationClient.cpp
@@ -105,7 +105,10 @@ namespace UE::MultiUserClient
 		// ... so add back the staged objects
 		for (const FSoftObjectPath& StagedObject : StagedObjects)
 		{
-			GetClientContent()->Stream->ReplicationMap.ReplicatedObjects.Add(StagedObject);
+			if (const UObject* Object = StagedObject.ResolveObject())
+			{
+				GetClientContent()->Stream->ReplicationMap.ReplicatedObjects.Add(StagedObject) = FConcertReplicatedObjectInfo::Make(*Object);
+			}
 		}
 		// To remove the hierarchy, the user must click the actor and delete it explicitly, which will call IEditableReplicationStream::RemoveObjects on the staged objects.
 		
diff --git a/Engine/Plugins/Developer/Concert/ConcertApp/MultiUserClient/Source/MultiUserClient/Private/Widgets/ActiveSession/Replication/Client/Multi/SMultiClientView.cpp b/Engine/Plugins/Developer/Concert/ConcertApp/MultiUserClient/Source/MultiUserClient/Private/Widgets/ActiveSession/Replication/Client/Multi/SMultiClientView.cpp
--- a/Engine/Plugins/Developer/Concert/ConcertApp/MultiUserClient/Source/MultiUserClient/Private/Widgets/ActiveSession/Replication/Client/Multi/SMultiClientView.cpp
+++ b/Engine/Plugins/Developer/Concert/ConcertApp/MultiUserClient/Source/MultiUserClient/Private/Widgets/ActiveSession/Replication/Client/Multi/SMultiClientView.cpp
@@ -83,7 +83,17 @@ namespace UE::MultiUserClient
 			return StreamModel->GetEditableStreams().Contains(LocalStream) ? LocalStream.ToSharedPtr() : nullptr;
 		});
 		
-		ObjectHierarchy = ConcertClientSharedSlate::CreateObjectHierarchyForComponentHierarchy();
+		ConcertClientSharedSlate::FFilterablePropertyTreeViewParams TreeViewParams
+		{
+			.PropertyColumns =
+			{
+				ReplicationColumns::Property::LabelColumn(),
+				ReplicationColumns::Property::TypeColumn(),
+				MultiStreamColumns::AssignPropertyColumn(MultiStreamEditorAttribute, InConcertClient, InClientManager)
+			},
+			.PrimaryPropertySort = { MultiStreamColumns::AssignPropertyColumnId, EColumnSortMode::Ascending }
+		};
+		TSharedRef<IPropertyTreeView> PropertyTreeView = CreateFilterablePropertyTreeView(MoveTemp(TreeViewParams));
 		
 		FCreateMultiStreamEditorParams Params
 		{
@@ -91,22 +101,24 @@ namespace UE::MultiUserClient
 			.ConsolidatedObjectModel = ConcertClientSharedSlate::CreateTransactionalStreamModel(),
 			.ObjectSource = MakeShared<ConcertClientSharedSlate::FActorSelectionSourceModel>(),
 			.PropertySource = MakeShared<ConcertClientSharedSlate::FSelectPropertyFromUClassModel>(),
-			.GetAutoAssignToStreamDelegate = MoveTemp(GetAutoAssignTargetDelegate),
-			.ViewerParams 
+			.GetAutoAssignToStreamDelegate = MoveTemp(GetAutoAssignTargetDelegate)
+		};
+		
+		ObjectHierarchy = ConcertClientSharedSlate::CreateObjectHierarchyForComponentHierarchy();
+		FCreateViewerParams ViewerParams
+		{
+			.PropertyTreeView = MoveTemp(PropertyTreeView),
+			.ObjectHierarchy = ObjectHierarchy, // This makes actors have children in the top view
+			.NameModel = ConcertClientSharedSlate::CreateEditorObjectNameModel(), // This makes actors use their labels, and components use the names given in the BP editor
+			.OnExtendObjectsContextMenu = FExtendObjectMenu::CreateSP(this, &SMultiClientView::ExtendObjectContextMenu),
+			.AdditionalObjectColumns =
 			{
-				.ObjectHierarchy = ObjectHierarchy, // This makes actors have children in the top view
-				.NameModel = ConcertClientSharedSlate::CreateEditorObjectNameModel(), // This makes actors use their labels, and components use the names given in the BP editor
-				.OnExtendObjectsContextMenu = FExtendObjectMenu::CreateSP(this, &SMultiClientView::ExtendObjectContextMenu),
-				.AdditionalObjectColumns =
-				{
-					MultiStreamColumns::ReplicationToggle(InConcertClient, ObjecHierarchyAttribute, InClientManager),
-					MultiStreamColumns::ReassignOwnership(InConcertClient, MultiStreamEditorAttribute, ObjecHierarchyAttribute, InClientManager.GetReassignmentLogic(), InClientManager)
-				},
-				.AdditionalPropertyColumns = { MultiStreamColumns::AssignPropertyColumn(MultiStreamEditorAttribute, InConcertClient, InClientManager) },
-				.PrimaryPropertySort = { MultiStreamColumns::AssignPropertyColumnId, EColumnSortMode::Ascending}
+				MultiStreamColumns::ReplicationToggle(InConcertClient, ObjecHierarchyAttribute, InClientManager),
+				MultiStreamColumns::ReassignOwnership(InConcertClient, MultiStreamEditorAttribute, ObjecHierarchyAttribute, InClientManager.GetReassignmentLogic(), InClientManager)
 			}
 		};
-		StreamEditor = CreateBaseMultiStreamEditor(MoveTemp(Params));
+		
+		StreamEditor = CreateBaseMultiStreamEditor(MoveTemp(Params), MoveTemp(ViewerParams));
 		check(StreamEditor);
 		return StreamEditor.ToSharedRef();
 	}
diff --git a/Engine/Plugins/Developer/Concert/ConcertApp/MultiUserClient/Source/MultiUserClient/Private/Widgets/ActiveSession/Replication/Client/Single/SReplicationClientView.cpp b/Engine/Plugins/Developer/Concert/ConcertApp/MultiUserClient/Source/MultiUserClient/Private/Widgets/ActiveSession/Replication/Client/Single/SReplicationClientView.cpp
--- a/Engine/Plugins/Developer/Concert/ConcertApp/MultiUserClient/Source/MultiUserClient/Private/Widgets/ActiveSession/Replication/Client/Single/SReplicationClientView.cpp
+++ b/Engine/Plugins/Developer/Concert/ConcertApp/MultiUserClient/Source/MultiUserClient/Private/Widgets/ActiveSession/Replication/Client/Single/SReplicationClientView.cpp
@@ -69,33 +69,33 @@ namespace UE::MultiUserClient
 			TAttribute<const IReplicationStreamViewer*>::CreateLambda([this](){ return EditorView.Get(); });
 
 		// Add checkboxes in front of top level and subobject rows for changing authority
-		const FCreateEditorParams ReplicationEditorCreationParams
+		ConcertClientSharedSlate::FDefaultStreamEditorParams DefaultEditorParams
 		{
-			.DataModel = InReplicationClient.GetClientEditModel(),
-			.ObjectSource = MakeShared<ConcertClientSharedSlate::FActorSelectionSourceModel>(),
-			.PropertySource = MakeShared<ConcertClientSharedSlate::FSelectPropertyFromUClassModel>(),
-			.IsEditingEnabled = TAttribute<bool>::CreateLambda([&SubmissionWorkflow](){ return CanEverSubmit(SubmissionWorkflow.GetUploadability()); }),
-			.EditingDisabledToolTipText = LOCTEXT("Editing.NotImplemented", "Editing remote clients is not implemented. You can only edit the local client."),
-			.ViewerParams =
+			.BaseEditorParams =
 			{
-				.ObjectHierarchy = ConcertClientSharedSlate::CreateObjectHierarchyForComponentHierarchy(), // This makes actors have children in the top view
-				.NameModel = ConcertClientSharedSlate::CreateEditorObjectNameModel(), // This makes actors use their labels, and components use the names given in the BP editor
-				.OnExtendObjectsContextMenu = FExtendObjectMenu::CreateSP(this, &SReplicationClientView::ExtendObjectContextMenu),
-				.AdditionalObjectColumns =
-				{
-					SingleClientColumns::ToggleObjectAuthority(AuthorityTracker, SubmissionWorkflow),
-					SingleClientColumns::ConflictWarningForObject(ConcertClient.ToSharedRef(), AuthorityCache, InReplicationClient.GetEndpointId()),
-					SingleClientColumns::OwnerOfObject(ConcertClient.ToSharedRef(), AuthorityCache)
-				},
-				.AdditionalPropertyColumns =
-				{
-					SingleClientColumns::OwnerOfProperty(ConcertClient.ToSharedRef(), AuthorityCache, GetReplicationViewerAttribute),
-					SingleClientColumns::ConflictWarningForProperty(ConcertClient.ToSharedRef(), GetReplicationViewerAttribute, AuthorityCache, InReplicationClient.GetEndpointId())
-				},
-			}
+				.DataModel = InReplicationClient.GetClientEditModel(),
+				.ObjectSource = MakeShared<ConcertClientSharedSlate::FActorSelectionSourceModel>(),
+				.PropertySource = MakeShared<ConcertClientSharedSlate::FSelectPropertyFromUClassModel>(),
+				.IsEditingEnabled = TAttribute<bool>::CreateLambda([&SubmissionWorkflow](){ return CanEverSubmit(SubmissionWorkflow.GetUploadability()); }),
+				.EditingDisabledToolTipText = LOCTEXT("Editing.NotImplemented", "Editing remote clients is not implemented. You can only edit the local client."),
+			},
+			.AdditionalPropertyColumns =
+			{
+				SingleClientColumns::OwnerOfProperty(ConcertClient.ToSharedRef(), AuthorityCache, GetReplicationViewerAttribute),
+				SingleClientColumns::ConflictWarningForProperty(ConcertClient.ToSharedRef(), GetReplicationViewerAttribute, AuthorityCache, InReplicationClient.GetEndpointId())
+			},
+			.ObjectHierarchy = ConcertClientSharedSlate::CreateObjectHierarchyForComponentHierarchy(), // This makes actors have children in the top view
+            .NameModel = ConcertClientSharedSlate::CreateEditorObjectNameModel(), // This makes actors use their labels, and components use the names given in the BP editor
+            .OnExtendObjectsContextMenu = FExtendObjectMenu::CreateSP(this, &SReplicationClientView::ExtendObjectContextMenu),
+            .AdditionalObjectColumns =
+            {
+            	SingleClientColumns::ToggleObjectAuthority(AuthorityTracker, SubmissionWorkflow),
+            	SingleClientColumns::ConflictWarningForObject(ConcertClient.ToSharedRef(), AuthorityCache, InReplicationClient.GetEndpointId()),
+            	SingleClientColumns::OwnerOfObject(ConcertClient.ToSharedRef(), AuthorityCache)
+            },
 		};
 
-		EditorView = ConcertClientSharedSlate::CreateDefaultStreamEditor(ReplicationEditorCreationParams);
+		EditorView = ConcertClientSharedSlate::CreateDefaultStreamEditor(MoveTemp(DefaultEditorParams));
 		return EditorView.ToSharedRef();
 	}
 
diff --git a/Engine/Plugins/Developer/Concert/ConcertApp/MultiUserClient/Source/MultiUserReplicationEditor/Private/StreamEditor/ReplicationSessionPresetEditorToolkit.cpp b/Engine/Plugins/Developer/Concert/ConcertApp/MultiUserClient/Source/MultiUserReplicationEditor/Private/StreamEditor/ReplicationSessionPresetEditorToolkit.cpp
--- a/Engine/Plugins/Developer/Concert/ConcertApp/MultiUserClient/Source/MultiUserReplicationEditor/Private/StreamEditor/ReplicationSessionPresetEditorToolkit.cpp
+++ b/Engine/Plugins/Developer/Concert/ConcertApp/MultiUserClient/Source/MultiUserReplicationEditor/Private/StreamEditor/ReplicationSessionPresetEditorToolkit.cpp
@@ -83,7 +83,15 @@ namespace UE::MultiUserReplicationEditor
 		
 		const TSharedRef<FActorSelectionSourceModel> ObjectSourceModel = MakeShared<FActorSelectionSourceModel>();
 		const TSharedRef<FSelectPropertyFromUClassModel> PropertySourceModel = MakeShared<FSelectPropertyFromUClassModel>();
-		const TSharedRef<IReplicationStreamEditor> EditorView = CreateDefaultStreamEditor(FCreateEditorParams{ AssetReadWriteModel, ObjectSourceModel, PropertySourceModel });
+		FDefaultStreamEditorParams DefaultParams{
+			.BaseEditorParams =
+			{
+				.DataModel = AssetReadWriteModel,
+				.ObjectSource = ObjectSourceModel,
+				.PropertySource = PropertySourceModel
+			}
+		};
+		const TSharedRef<IReplicationStreamEditor> EditorView = CreateDefaultStreamEditor(MoveTemp(DefaultParams));
 		return SNew(SDockTab)
 			.Label(LOCTEXT("BaseDetailsTitle", "Details"))
 			[
diff --git a/Engine/Plugins/Developer/Concert/ConcertUI/ConcertClientSharedSlate/Source/ConcertClientSharedSlate/Private/Replication/ClientReplicationWidgetFactories.cpp b/Engine/Plugins/Developer/Concert/ConcertUI/ConcertClientSharedSlate/Source/ConcertClientSharedSlate/Private/Replication/ClientReplicationWidgetFactories.cpp
--- a/Engine/Plugins/Developer/Concert/ConcertUI/ConcertClientSharedSlate/Source/ConcertClientSharedSlate/Private/Replication/ClientReplicationWidgetFactories.cpp
+++ b/Engine/Plugins/Developer/Concert/ConcertUI/ConcertClientSharedSlate/Source/ConcertClientSharedSlate/Private/Replication/ClientReplicationWidgetFactories.cpp
@@ -2,12 +2,15 @@
 
 #include "Replication/ClientReplicationWidgetFactories.h"
 
+#include "Editor/View/ClientEditorColumns.h"
+#include "Editor/View/PropertyTree/SFilteredPropertyTreeView.h"
 #include "Replication/ReplicationWidgetFactories.h"
 #include "Replication/Editor/Model/Object/EditorObjectHierarchyModel.h"
 #include "Replication/Editor/Model/Object/EditorObjectNameModel.h"
 #include "Replication/Editor/Model/ReplicationStreamObject.h"
-#include "Replication/Editor/Model/TransactionalReplicationStreamModel.h"
 #include "Replication/Editor/View/ObjectEditor/SDefaultReplicationStreamEditor.h"
+#include "Replication/Editor/View/SelectionViewerColumns.h"
+#include "Replication/Editor/Model/TransactionalReplicationStreamModel.h"
 
 #include "UObject/UObjectGlobals.h"
 #include "UObject/Package.h"
@@ -52,8 +55,58 @@ namespace UE::ConcertClientSharedSlate
 		return CreateTransactionalStreamModel(ConcertSharedSlate::CreateBaseStreamModel(MoveTemp(Attribute)), *Object);
 	}
 	
-	TSharedRef<ConcertSharedSlate::IReplicationStreamEditor> CreateDefaultStreamEditor(ConcertSharedSlate::FCreateEditorParams Params)
+	TSharedRef<ConcertSharedSlate::IPropertyTreeView> CreateFilterablePropertyTreeView(FFilterablePropertyTreeViewParams Params)
 	{
-		return SNew(SDefaultReplicationStreamEditor, MoveTemp(Params));
+		return SNew(SFilteredPropertyTreeView, MoveTemp(Params));
+	}
+	
+	TSharedRef<ConcertSharedSlate::IReplicationStreamEditor> CreateDefaultStreamEditor(FDefaultStreamEditorParams Params)
+	{
+		using namespace ConcertSharedSlate;
+		using namespace ConcertSharedSlate::ReplicationColumns;
+
+		// This is a hack.
+		// The architecturally correct way to fix is pass FReplicationPropertyColumn the FSoftObjectPath to the object for which the column is being constructed.
+		struct FEditorIndirection
+		{
+			TSharedPtr<IReplicationStreamEditor> Editor;
+		};
+		TSharedRef<FEditorIndirection> Indirection = MakeShared<FEditorIndirection>();
+		
+		const FReplicationPropertyColumn ReplicatesColumn = ReplicationColumns::Property::ReplicatesColumns(
+			TAttribute<IReplicationStreamViewer*>::CreateLambda([Indirection](){ return Indirection->Editor.Get(); }),
+			Params.BaseEditorParams.DataModel,
+			TReplicationColumnDelegates<FReplicatedPropertyData>::FIsEnabled::CreateLambda([IsEnabled = Params.BaseEditorParams.IsEditingEnabled](const FReplicatedPropertyData&)
+			{
+				return !IsEnabled.IsBound() || IsEnabled.Get();
+			}),
+			Params.BaseEditorParams.EditingDisabledToolTipText
+			);
+		
+		TArray<FReplicationPropertyColumn>& PropertyColumns = Params.AdditionalPropertyColumns;
+		const bool bHasType = PropertyColumns.ContainsByPredicate([](const FReplicationPropertyColumn& Column)
+		{
+			return Column.ColumnId == Property::TypeColumnId;
+		});
+		if (!bHasType)
+		{
+			PropertyColumns.Add(Property::TypeColumn());
+		}
+		PropertyColumns.Add(ReplicatesColumn);
+
+		FCreateViewerParams ViewerParams
+		{
+			.PropertyTreeView = CreateFilterablePropertyTreeView({ .PropertyColumns = PropertyColumns }),
+			.ObjectHierarchy = MoveTemp(Params.ObjectHierarchy),
+			.NameModel = MoveTemp(Params.NameModel),
+			.OnExtendObjectsContextMenu = MoveTemp(Params.OnExtendObjectsContextMenu),
+			.AdditionalObjectColumns = MoveTemp(Params.AdditionalObjectColumns),
+			.PrimaryObjectSort = FColumnSortInfo{ TopLevel::LabelColumnId, EColumnSortMode::Ascending },
+			.SecondaryObjectSort = FColumnSortInfo{ TopLevel::LabelColumnId, EColumnSortMode::Ascending },
+		};
+		
+		TSharedRef<SDefaultReplicationStreamEditor> Editor = SNew(SDefaultReplicationStreamEditor, MoveTemp(Params.BaseEditorParams), MoveTemp(ViewerParams));
+		Indirection->Editor = Editor;
+		return Editor;
 	}
 }
\ No newline at end of file
diff --git a/Engine/Plugins/Developer/Concert/ConcertUI/ConcertClientSharedSlate/Source/ConcertClientSharedSlate/Private/Replication/Editor/View/ClientEditorColumns.cpp b/Engine/Plugins/Developer/Concert/ConcertUI/ConcertClientSharedSlate/Source/ConcertClientSharedSlate/Private/Replication/Editor/View/ClientEditorColumns.cpp
--- a/Engine/Plugins/Developer/Concert/ConcertUI/ConcertClientSharedSlate/Source/ConcertClientSharedSlate/Private/Replication/Editor/View/ClientEditorColumns.cpp
+++ b/Engine/Plugins/Developer/Concert/ConcertUI/ConcertClientSharedSlate/Source/ConcertClientSharedSlate/Private/Replication/Editor/View/ClientEditorColumns.cpp
@@ -84,7 +84,7 @@ namespace UE::ConcertClientSharedSlate::ReplicationColumns::Property
 	}
 	
 	ConcertSharedSlate::ReplicationColumns::FReplicationPropertyColumn ReplicatesColumns(
-		TWeakPtr<ConcertSharedSlate::IReplicationStreamViewer> Viewer,
+		TAttribute<ConcertSharedSlate::IReplicationStreamViewer*> Viewer,
 		TWeakPtr<ConcertSharedSlate::IEditableReplicationStreamModel> Model,
 		ConcertSharedSlate::TReplicationColumnDelegates<ConcertSharedSlate::FReplicatedPropertyData>::FIsEnabled IsEnabledDelegate,
 		TAttribute<FText> DisabledToolTipText,
@@ -100,14 +100,14 @@ namespace UE::ConcertClientSharedSlate::ReplicationColumns::Property
 				FPropertyColumnDelegates::FGetColumnCheckboxState::CreateLambda(
 				[Viewer, Model](const FReplicatedPropertyData& Data)
 				{
-					const TSharedPtr<IReplicationStreamViewer> ViewerPin = Viewer.Pin();
+					const IReplicationStreamViewer* ViewerPin = Viewer.Get();
 					const TSharedPtr<IEditableReplicationStreamModel> ModelPin = Model.Pin();
 					return ensure(ViewerPin && ModelPin) ? Private::OnGetPropertyCheckboxState(Data.GetProperty(), *ViewerPin, *ModelPin) : ECheckBoxState::Undetermined;
 				}),
 				FPropertyColumnDelegates::FOnColumnCheckboxChanged::CreateLambda(
 				[Viewer, Model](bool bIsChecked, const FReplicatedPropertyData& Data)
 				{
-					const TSharedPtr<IReplicationStreamViewer> ViewerPin = Viewer.Pin();
+					const IReplicationStreamViewer* ViewerPin = Viewer.Get();
 					const TSharedPtr<IEditableReplicationStreamModel> ModelPin = Model.Pin();
 					if (ensure(ViewerPin && ModelPin))
 					{
diff --git a/Engine/Plugins/Developer/Concert/ConcertUI/ConcertClientSharedSlate/Source/ConcertClientSharedSlate/Private/Replication/Editor/View/ClientEditorColumns.h b/Engine/Plugins/Developer/Concert/ConcertUI/ConcertClientSharedSlate/Source/ConcertClientSharedSlate/Private/Replication/Editor/View/ClientEditorColumns.h
--- a/Engine/Plugins/Developer/Concert/ConcertUI/ConcertClientSharedSlate/Source/ConcertClientSharedSlate/Private/Replication/Editor/View/ClientEditorColumns.h
+++ b/Engine/Plugins/Developer/Concert/ConcertUI/ConcertClientSharedSlate/Source/ConcertClientSharedSlate/Private/Replication/Editor/View/ClientEditorColumns.h
@@ -36,7 +36,7 @@ namespace UE::ConcertClientSharedSlate::ReplicationColumns::Property
 	 * @param Priority Determines position of this columns relative to the others
 	 */
 	ConcertSharedSlate::ReplicationColumns::FReplicationPropertyColumn ReplicatesColumns(
-		TWeakPtr<ConcertSharedSlate::IReplicationStreamViewer> Viewer,
+		TAttribute<ConcertSharedSlate::IReplicationStreamViewer*> Viewer,
 		TWeakPtr<ConcertSharedSlate::IEditableReplicationStreamModel> Model,
 		ConcertSharedSlate::TReplicationColumnDelegates<ConcertSharedSlate::FReplicatedPropertyData>::FIsEnabled IsEnabledDelegate = {},
 		TAttribute<FText> DisabledToolTipText = {},
diff --git a/Engine/Plugins/Developer/Concert/ConcertUI/ConcertClientSharedSlate/Source/ConcertClientSharedSlate/Private/Replication/Editor/View/ObjectEditor/SDefaultReplicationStreamEditor.cpp b/Engine/Plugins/Developer/Concert/ConcertUI/ConcertClientSharedSlate/Source/ConcertClientSharedSlate/Private/Replication/Editor/View/ObjectEditor/SDefaultReplicationStreamEditor.cpp
--- a/Engine/Plugins/Developer/Concert/ConcertUI/ConcertClientSharedSlate/Source/ConcertClientSharedSlate/Private/Replication/Editor/View/ObjectEditor/SDefaultReplicationStreamEditor.cpp
+++ b/Engine/Plugins/Developer/Concert/ConcertUI/ConcertClientSharedSlate/Source/ConcertClientSharedSlate/Private/Replication/Editor/View/ObjectEditor/SDefaultReplicationStreamEditor.cpp
@@ -3,41 +3,14 @@
 #include "SDefaultReplicationStreamEditor.h"
 
 #include "Replication/ReplicationWidgetFactories.h"
-#include "Replication/Editor/View/ClientEditorColumns.h"
-#include "Replication/Editor/View/SelectionViewerColumns.h"
 
 namespace UE::ConcertClientSharedSlate
 {
-	void SDefaultReplicationStreamEditor::Construct(const FArguments& InArgs, ConcertSharedSlate::FCreateEditorParams EditorParams)
+	void SDefaultReplicationStreamEditor::Construct(const FArguments& InArgs, ConcertSharedSlate::FCreateEditorParams EditorParams, ConcertSharedSlate::FCreateViewerParams ViewerParams)
 	{
-		using namespace ConcertSharedSlate;
-		using namespace ConcertSharedSlate::ReplicationColumns;
-		using namespace ConcertSharedSlate::ReplicationColumns::Property;
-		using namespace ConcertClientSharedSlate::ReplicationColumns::Property;
-		
 		PropertiesModel = EditorParams.DataModel;
 
-		const FReplicationPropertyColumn ReplicatesColumn = ReplicatesColumns(
-			SharedThis(this),
-			PropertiesModel,
-			TReplicationColumnDelegates<FReplicatedPropertyData>::FIsEnabled::CreateLambda([IsEnabled = EditorParams.IsEditingEnabled](const FReplicatedPropertyData&)
-			{
-				return !IsEnabled.IsBound() || IsEnabled.Get();
-			}),
-			EditorParams.EditingDisabledToolTipText
-			);
-		TArray<FReplicationPropertyColumn>& PropertyColumns = EditorParams.ViewerParams.AdditionalPropertyColumns;
-		PropertyColumns.Add(ReplicatesColumn);
-
-		// Set both primary and secondary in case one is overriden but always use the override.
-		EditorParams.ViewerParams.PrimaryPropertySort = EditorParams.ViewerParams.PrimaryPropertySort.IsValid()
-			? EditorParams.ViewerParams.PrimaryPropertySort
-			: FColumnSortInfo{ ReplicatesColumnId, EColumnSortMode::Ascending };
-		EditorParams.ViewerParams.SecondaryPropertySort = EditorParams.ViewerParams.SecondaryPropertySort.IsValid()
-			? EditorParams.ViewerParams.SecondaryPropertySort
-			: FColumnSortInfo{ LabelColumnId, EColumnSortMode::Ascending };
-
-		WrappedEditor = CreateBaseStreamEditor(MoveTemp(EditorParams));
+		WrappedEditor = CreateBaseStreamEditor(MoveTemp(EditorParams), MoveTemp(ViewerParams));
 		ChildSlot
 		[
 			WrappedEditor.ToSharedRef()
diff --git a/Engine/Plugins/Developer/Concert/ConcertUI/ConcertClientSharedSlate/Source/ConcertClientSharedSlate/Private/Replication/Editor/View/ObjectEditor/SDefaultReplicationStreamEditor.h b/Engine/Plugins/Developer/Concert/ConcertUI/ConcertClientSharedSlate/Source/ConcertClientSharedSlate/Private/Replication/Editor/View/ObjectEditor/SDefaultReplicationStreamEditor.h
--- a/Engine/Plugins/Developer/Concert/ConcertUI/ConcertClientSharedSlate/Source/ConcertClientSharedSlate/Private/Replication/Editor/View/ObjectEditor/SDefaultReplicationStreamEditor.h
+++ b/Engine/Plugins/Developer/Concert/ConcertUI/ConcertClientSharedSlate/Source/ConcertClientSharedSlate/Private/Replication/Editor/View/ObjectEditor/SDefaultReplicationStreamEditor.h
@@ -16,6 +16,7 @@ struct FConcertStreamObjectAutoBindingRules;
 namespace UE::ConcertSharedSlate
 {
 	struct FCreateEditorParams;
+	struct FCreateViewerParams;
 	class IEditableReplicationStreamModel;
 }
 
@@ -34,7 +35,7 @@ namespace UE::ConcertClientSharedSlate
 		SLATE_BEGIN_ARGS(SDefaultReplicationStreamEditor) {}
 		SLATE_END_ARGS()
 
-		void Construct(const FArguments& InArgs, ConcertSharedSlate::FCreateEditorParams EditorParams);
+		void Construct(const FArguments& InArgs, ConcertSharedSlate::FCreateEditorParams EditorParams, ConcertSharedSlate::FCreateViewerParams ViewerParams);
 
 		//~ Begin IReplicationStreamEditor Interface
 		virtual void Refresh() override;
diff --git a/Engine/Plugins/Developer/Concert/ConcertUI/ConcertClientSharedSlate/Source/ConcertClientSharedSlate/Private/Replication/Editor/View/PropertyTree/PropertyFilterBase.h b/Engine/Plugins/Developer/Concert/ConcertUI/ConcertClientSharedSlate/Source/ConcertClientSharedSlate/Private/Replication/Editor/View/PropertyTree/PropertyFilterBase.h
new file mode 100644
--- /dev/null
+++ b/Engine/Plugins/Developer/Concert/ConcertUI/ConcertClientSharedSlate/Source/ConcertClientSharedSlate/Private/Replication/Editor/View/PropertyTree/PropertyFilterBase.h
@@ -0,0 +1,42 @@
+// Copyright Epic Games, Inc. All Rights Reserved.
+
+#pragma once
+
+#include "Containers/Set.h"
+#include "Misc/IFilter.h"
+
+class FFieldClass;
+namespace UE::ConcertSharedSlate { class FReplicatedPropertyData; }
+
+namespace UE::ConcertClientSharedSlate
+{
+	/**
+	 * Base class for property filters.
+	 * 
+	 * This filter inverts the results: all filters are supposed to be inverse filters, i.e. they should be run when
+	 * they are greyed out in the UI. When they are run, they should remove specific properties.
+	 *
+	 * Subclasses simply implement Matches, which figures out whether the property is contained in some FReplicatedPropertyData.
+	 */
+	class FPropertyFilterBase : public IFilter<const ConcertSharedSlate::FReplicatedPropertyData&>
+	{
+	public:
+		
+		//~ Begin IFilter Interface
+		virtual bool PassesFilter(const ConcertSharedSlate::FReplicatedPropertyData& InItem) const final override
+		{
+			return MatchesFilteredForProperty(InItem);
+		}
+		virtual FChangedEvent& OnChanged() final override { return ChangedEventDelegate; }
+		//~ End IFilter Interface
+
+	private:
+
+		/** Dummy delegate to fulfill IFilter interface: this filter will never change so this delegate will never trigger. */
+		FChangedEvent ChangedEventDelegate;
+
+		/** @return Whether this item contains the property this filter is looking for.*/
+		virtual bool MatchesFilteredForProperty(const ConcertSharedSlate::FReplicatedPropertyData& InItem) const = 0;
+	};
+}
+
diff --git a/Engine/Plugins/Developer/Concert/ConcertUI/ConcertClientSharedSlate/Source/ConcertClientSharedSlate/Private/Replication/Editor/View/PropertyTree/PropertyFilter_ByPropertyType.cpp b/Engine/Plugins/Developer/Concert/ConcertUI/ConcertClientSharedSlate/Source/ConcertClientSharedSlate/Private/Replication/Editor/View/PropertyTree/PropertyFilter_ByPropertyType.cpp
new file mode 100644
--- /dev/null
+++ b/Engine/Plugins/Developer/Concert/ConcertUI/ConcertClientSharedSlate/Source/ConcertClientSharedSlate/Private/Replication/Editor/View/PropertyTree/PropertyFilter_ByPropertyType.cpp
@@ -0,0 +1,23 @@
+// Copyright Epic Games, Inc. All Rights Reserved.
+
+#include "PropertyFilter_ByPropertyType.h"
+
+#include "Replication/PropertyChainUtils.h"
+#include "Replication/Editor/Model/ReplicatedPropertyData.h"
+
+#include "UObject/UnrealType.h"
+
+namespace UE::ConcertClientSharedSlate
+{
+	bool FPropertyFilter_ByPropertyType::MatchesFilteredForProperty(const ConcertSharedSlate::FReplicatedPropertyData& InItem) const
+	{
+		UClass* Class = InItem.GetOwningClass().TryLoadClass<UObject>();
+		if (!Class)
+		{
+			return false;
+		}
+
+		const FProperty* Property = ConcertSyncCore::PropertyChain::ResolveProperty(*Class, InItem.GetProperty());
+		return Property && AllowedClasses.Contains(Property->GetClass());
+	}
+}
diff --git a/Engine/Plugins/Developer/Concert/ConcertUI/ConcertClientSharedSlate/Source/ConcertClientSharedSlate/Private/Replication/Editor/View/PropertyTree/PropertyFilter_ByPropertyType.h b/Engine/Plugins/Developer/Concert/ConcertUI/ConcertClientSharedSlate/Source/ConcertClientSharedSlate/Private/Replication/Editor/View/PropertyTree/PropertyFilter_ByPropertyType.h
new file mode 100644
--- /dev/null
+++ b/Engine/Plugins/Developer/Concert/ConcertUI/ConcertClientSharedSlate/Source/ConcertClientSharedSlate/Private/Replication/Editor/View/PropertyTree/PropertyFilter_ByPropertyType.h
@@ -0,0 +1,36 @@
+// Copyright Epic Games, Inc. All Rights Reserved.
+
+#pragma once
+
+#include "PropertyFilterBase.h"
+#include "Containers/Set.h"
+#include "Misc/IFilter.h"
+
+class FFieldClass;
+
+namespace UE::ConcertSharedSlate { class FReplicatedPropertyData; }
+
+namespace UE::ConcertClientSharedSlate
+{
+	/** Filters a property based on the FFieldClass it has. */
+	class FPropertyFilter_ByPropertyType : public FPropertyFilterBase
+	{
+	public:
+
+		FPropertyFilter_ByPropertyType(TSet<FFieldClass*> AllowedClasses)
+			: AllowedClasses(MoveTemp(AllowedClasses))
+		{}
+
+	private:
+
+		TSet<FFieldClass*> AllowedClasses;
+
+		/** Dummy delegate to fulfill IFilter interface: this filter will never change so this delegate will never trigger. */
+		FChangedEvent ChangedEventDelegate;
+
+		//~ Begin FPropertyFilterBase Interface
+		virtual bool MatchesFilteredForProperty(const ConcertSharedSlate::FReplicatedPropertyData& InItem) const override;
+		//~ End FPropertyFilterBase Interface
+	};
+}
+
diff --git a/Engine/Plugins/Developer/Concert/ConcertUI/ConcertClientSharedSlate/Source/ConcertClientSharedSlate/Private/Replication/Editor/View/PropertyTree/PropertyFrontendFilter.h b/Engine/Plugins/Developer/Concert/ConcertUI/ConcertClientSharedSlate/Source/ConcertClientSharedSlate/Private/Replication/Editor/View/PropertyTree/PropertyFrontendFilter.h
new file mode 100644
--- /dev/null
+++ b/Engine/Plugins/Developer/Concert/ConcertUI/ConcertClientSharedSlate/Source/ConcertClientSharedSlate/Private/Replication/Editor/View/PropertyTree/PropertyFrontendFilter.h
@@ -0,0 +1,80 @@
+// Copyright Epic Games, Inc. All Rights Reserved.
+
+#pragma once
+
+#include "Filters/FilterBase.h"
+#include "Internationalization/Text.h"
+
+namespace UE::ConcertSharedSlate { class FReplicatedPropertyData; }
+
+namespace UE::ConcertClientSharedSlate
+{
+	/** Base UI filter for filtering properties. Gets rid of all the extra functionality in FFilterBase we do not need. */
+	class FPropertyFrontendFilter : public FFilterBase<const ConcertSharedSlate::FReplicatedPropertyData&>
+	{
+	public:
+
+		FPropertyFrontendFilter(TSharedRef<FFilterCategory> InCategory, FText DisplayName, FText Tooltip = FText::GetEmpty())
+			: FFilterBase(MoveTemp(InCategory))
+			, DisplayName(MoveTemp(DisplayName))
+			, Tooltip(MoveTemp(Tooltip))
+		{}
+
+		//~ Begin FFilterBase Interface
+		virtual FString GetName() const override { return DisplayName.ToString(); }
+		virtual FText GetDisplayName() const override { return DisplayName; }
+		virtual FText GetToolTipText() const override { return Tooltip; }
+		virtual FLinearColor GetColor() const override { return FLinearColor(0.6f, 0.6f, 0.6f, 0.6f); }
+		virtual FName GetIconName() const override { return NAME_None; }
+		virtual bool IsInverseFilter() const override { return false; }
+		virtual void ActiveStateChanged(bool bActive) override {}
+		virtual void ModifyContextMenu(FMenuBuilder& MenuBuilder) override {}
+		virtual void SaveSettings(const FString& IniFilename, const FString& IniSection, const FString& SettingsString) const override {}
+		virtual void LoadSettings(const FString& IniFilename, const FString& IniSection, const FString& SettingsString) override {}
+		//~ End FFilterBase Interface
+
+	private:
+
+		FText DisplayName;
+		FText Tooltip;
+	};
+
+	/** Inlines a IFilter<const FReplicatedPropertyData&> at compile time to avoid an additional TSharedPtr<IFilter<const FReplicatedPropertyData&>>. */
+	template<typename TFilterType>
+	class TPropertyFrontendFilter : public FPropertyFrontendFilter
+	{
+	public:
+
+		template<typename... TArg>
+		TPropertyFrontendFilter(TSharedRef<FFilterCategory> InCategory, FText DisplayName, TArg&&... Arg)
+			: FPropertyFrontendFilter(MoveTemp(InCategory), MoveTemp(DisplayName))
+			, FilterImplementation(Forward<TArg>(Arg)...)
+		{
+			FilterImplementation.OnChanged().AddLambda([this]()
+			{
+				FPropertyFrontendFilter::BroadcastChangedEvent();
+			});
+		}
+
+		template<typename... TArg>
+		TPropertyFrontendFilter(TSharedRef<FFilterCategory> InCategory, FText DisplayName, FText Tooltip, TArg&&... Arg)
+			: FPropertyFrontendFilter(MoveTemp(InCategory), MoveTemp(DisplayName), MoveTemp(Tooltip))
+			, FilterImplementation(Forward<TArg>(Arg)...)
+		{
+			FilterImplementation.OnChanged().AddLambda([this]()
+			{
+				FPropertyFrontendFilter::BroadcastChangedEvent();
+			});
+		}
+		
+		virtual bool PassesFilter(const ConcertSharedSlate::FReplicatedPropertyData& InItem) const override
+		{
+			return FilterImplementation.PassesFilter(InItem);
+		}
+
+	private:
+		
+		TFilterType FilterImplementation;
+	};
+}
+
diff --git a/Engine/Plugins/Developer/Concert/ConcertUI/ConcertClientSharedSlate/Source/ConcertClientSharedSlate/Private/Replication/Editor/View/PropertyTree/SFilteredPropertyTreeView.cpp b/Engine/Plugins/Developer/Concert/ConcertUI/ConcertClientSharedSlate/Source/ConcertClientSharedSlate/Private/Replication/Editor/View/PropertyTree/SFilteredPropertyTreeView.cpp
new file mode 100644
--- /dev/null
+++ b/Engine/Plugins/Developer/Concert/ConcertUI/ConcertClientSharedSlate/Source/ConcertClientSharedSlate/Private/Replication/Editor/View/PropertyTree/SFilteredPropertyTreeView.cpp
@@ -0,0 +1,178 @@
+// Copyright Epic Games, Inc. All Rights Reserved.
+
+#include "SFilteredPropertyTreeView.h"
+
+#include "Filters/SBasicFilterBar.h"
+#include "PropertyFilter_ByPropertyType.h"
+#include "PropertyFrontendFilter.h"
+#include "Replication/ClientReplicationWidgetFactories.h"
+#include "Replication/ReplicationWidgetFactories.h"
+#include "Replication/Editor/Model/ReplicatedPropertyData.h"
+
+#include "UObject/UnrealType.h"
+
+#define LOCTEXT_NAMESPACE "SReplicatedPropertiesView"
+
+namespace UE::ConcertClientSharedSlate
+{
+	/** Exposes SetFrontendFilterActive so we can manually enable the default filters */
+	class SReplicationFilterBar : public SBasicFilterBar<const ConcertSharedSlate::FReplicatedPropertyData&>
+	{
+		using Super = SBasicFilterBar<const ConcertSharedSlate::FReplicatedPropertyData&>;
+	public:
+
+		SLATE_BEGIN_ARGS(SReplicationFilterBar)
+		{}
+			SLATE_EVENT(FOnFilterChanged, OnFilterChanged)
+		SLATE_END_ARGS()
+
+		void Construct(const FArguments& InArgs, TArray<SFilteredPropertyTreeView::FFilterRef> AllFilters)
+		{
+			Super::Construct(
+			Super::FArguments()
+				.FilterPillStyle(EFilterPillStyle::Basic)
+				.CustomFilters(MoveTemp(AllFilters))
+				.OnFilterChanged(InArgs._OnFilterChanged)
+				.UseSectionsForCategories(true)
+				);
+		}
+
+		// Expose from SBasicFilterBar
+		using Super::SetFrontendFilterActive;
+
+		void SetFilterVisuallyEnabled(const SFilteredPropertyTreeView::FFilterRef& Filter, bool bEnabled)
+		{
+			const TSharedRef<SFilter>* FilterWidget = Filters.FindByPredicate([&Filter](const TSharedRef<SFilter>& FilterWidget)
+			{
+				return FilterWidget->GetFrontendFilter() == Filter;
+			});
+			if (ensure(FilterWidget))
+			{
+				FilterWidget->Get().SetEnabled(bEnabled);
+			}
+		}
+	};
+	
+	void SFilteredPropertyTreeView::Construct(const FArguments& InArgs, FFilterablePropertyTreeViewParams Params)
+	{
+		const FBuildFilterBarResult Filters = BuildFilterBar();
+
+		using namespace UE::ConcertSharedSlate;
+		FCreatePropertyTreeViewParams TreeViewParams
+		{
+			.PropertyColumns = MoveTemp(Params.PropertyColumns),
+			.FilterItem = FFilterPropertyData::CreateSP(this, &SFilteredPropertyTreeView::PassesFilters),
+			.PrimaryPropertySort = Params.PrimaryPropertySort,
+			.SecondaryPropertySort = Params.SecondaryPropertySort
+		};
+		TreeViewParams.LeftOfPropertySearchBar.Widget = 
+			SNew(SHorizontalBox)
+
+			// The combo button for selecting the property filters
+			+SHorizontalBox::Slot()
+			.AutoWidth()
+			[
+				SBasicFilterBar<const FReplicatedPropertyData&>::MakeAddFilterButton(FilterBar.ToSharedRef())
+			]
+
+			+SHorizontalBox::Slot()
+			.AutoWidth()
+			[
+				InArgs._LeftOfSearchBar.Widget
+			]
+		;
+		TreeViewParams.RowBelowSearchBar.Widget = FilterBar.ToSharedRef();
+		TreeViewParams.NoItemsContent.Widget = SNew(STextBlock).Text(LOCTEXT("AllFitlered", "All properties filtered."));
+			
+		ExtendedTreeView = CreateSearchablePropertyTreeView(MoveTemp(TreeViewParams));
+		
+		ChildSlot
+		[
+			ExtendedTreeView->GetWidget()
+		];
+
+		// For better UX, hide certain properties by default (e.g. why would you want to replicate bools?)
+		// Do this AFTER initializing ReplicatedProperties because it triggers the OnItemsChanged callback.
+		for (const FFilterRef& Filter : Filters.DisabledByDefault)
+		{
+			FilterBar->SetFilterCheckState(Filter, ECheckBoxState::Unchecked);
+		}
+		// Show all the other filters as enabled (not greyed out: blue) - they will not run their logic since they are inverse.
+		for (const FFilterRef& Filter : Filters.EnabledByDefault)
+		{
+			// Makes it appear on the bar
+			FilterBar->SetFilterCheckState(Filter, ECheckBoxState::Checked);
+			// Functionally makes it affect the search
+			FilterBar->SetFrontendFilterActive(Filter, true);
+			// Visually makes the button blue (i.e. so it looks enabled)
+			FilterBar->SetFilterVisuallyEnabled(Filter, true);
+		}
+	}
+
+	SFilteredPropertyTreeView::FBuildFilterBarResult SFilteredPropertyTreeView::BuildFilterBar()
+	{
+		TSharedRef<FFilterCategory> CommonCategory = MakeShared<FFilterCategory>(
+			LOCTEXT("CommonCategory.Name", "Common"),
+			LOCTEXT("CommonCategory.ToolTip", "Include commonly replicated properties.")
+		);
+		TSharedRef<FFilterCategory> UncommonCategory = MakeShared<FFilterCategory>(
+			LOCTEXT("UncommonCategory.Advanced", "Uncommon"),
+			LOCTEXT("UncommonCategory.ToolTip", "Include uncommonly replicated properties.")
+		);
+
+		FBuildFilterBarResult Result;
+		using FFrontendFilter = TPropertyFrontendFilter<FPropertyFilter_ByPropertyType>;
+		// Do not show under search bar
+		Result.DisabledByDefault =
+		{
+			// Ordering matters for the drop-down menu next to settings
+			// Common
+			MakeShared<FFrontendFilter>(CommonCategory, LOCTEXT("Ints", "Integer"), LOCTEXT("Ints.Tooltip", "Includes: uint16, uint32, uint64, int16, int32, int64"), TSet<FFieldClass*>{ FUInt16Property::StaticClass(), FUInt32Property::StaticClass(), FUInt64Property::StaticClass(), FInt16Property::StaticClass(), FIntProperty::StaticClass(), FInt64Property::StaticClass() }),
+			MakeShared<FFrontendFilter>(CommonCategory, LOCTEXT("Floats", "Float"), LOCTEXT("Floats.Tooltip", "Includes: float, double"), TSet<FFieldClass*>{ FFloatProperty::StaticClass(), FDoubleProperty::StaticClass() }),
+			MakeShared<FFrontendFilter>(CommonCategory, LOCTEXT("Struct", "Struct"), TSet<FFieldClass*>{ FStructProperty::StaticClass() }),
+			MakeShared<FFrontendFilter>(CommonCategory, LOCTEXT("Containers", "Containers"), LOCTEXT("Containers.Tooltip", "Includes: array, set, map"), TSet<FFieldClass*>{ FArrayProperty::StaticClass(), FSetProperty::StaticClass(), FMapProperty::StaticClass() }),
+
+			// Uncommon
+			MakeShared<FFrontendFilter>(UncommonCategory, LOCTEXT("Bool", "Boolean"), TSet<FFieldClass*>{ FBoolProperty::StaticClass() }),
+			MakeShared<FFrontendFilter>(UncommonCategory, LOCTEXT("Enum", "Enum"), TSet<FFieldClass*>{ FEnumProperty::StaticClass(), FByteProperty::StaticClass() }),
+			MakeShared<FFrontendFilter>(UncommonCategory, LOCTEXT("Text", "Text"), LOCTEXT("Text.Tooltip", "Includes: FName, FString, FText"), TSet<FFieldClass*>{ FNameProperty::StaticClass(), FStrProperty::StaticClass(), FTextProperty::StaticClass() }),
+			MakeShared<FFrontendFilter>(UncommonCategory, LOCTEXT("SoftPtr", "Soft Ptr"), TSet<FFieldClass*>{ FSoftObjectProperty::StaticClass() }),
+		};
+		// Show up under search bar as enabled
+		Result.EnabledByDefault =
+		{
+			// We'll not enable any filters by default because that's the default for other places in the engine, like the Content Browser, too
+		};
+		// Show up in menu
+		TArray<FFilterRef> AllFilters;
+		AllFilters.Append(Result.EnabledByDefault);
+		AllFilters.Append(Result.DisabledByDefault);
+		
+		FilterBar = SNew(SReplicationFilterBar, MoveTemp(AllFilters))
+			.OnFilterChanged(this, &SFilteredPropertyTreeView::RequestRefilter);
+		
+		return Result;
+	}
+
+	ConcertSharedSlate::EFilterResult SFilteredPropertyTreeView::PassesFilters(const ConcertSharedSlate::FReplicatedPropertyData& ReplicatedPropertyData) const
+	{
+		const bool bPassesFilter = FilterBar->GetAllActiveFilters()->Num() == 0 // Return all items when none enabled
+			|| PassesAnyFilters(ReplicatedPropertyData);
+		return bPassesFilter ? ConcertSharedSlate::EFilterResult::PassesFilter : ConcertSharedSlate::EFilterResult::DoesNotPassFilter;
+	}
+
+	bool SFilteredPropertyTreeView::PassesAnyFilters(const ConcertSharedSlate::FReplicatedPropertyData& ReplicatedPropertyData) const
+	{
+		TSharedPtr<TFilterCollection<const ConcertSharedSlate::FReplicatedPropertyData&>> FilterCollection = FilterBar->GetAllActiveFilters();
+		for (int32 Index = 0; Index < FilterCollection->Num(); Index++)
+		{
+			if (FilterCollection->GetFilterAtIndex(Index)->PassesFilter(ReplicatedPropertyData))
+			{
+				return true;
+			}
+		}
+		return false;
+	}
+}
+
+#undef LOCTEXT_NAMESPACE
\ No newline at end of file
diff --git a/Engine/Plugins/Developer/Concert/ConcertUI/ConcertClientSharedSlate/Source/ConcertClientSharedSlate/Private/Replication/Editor/View/PropertyTree/SFilteredPropertyTreeView.h b/Engine/Plugins/Developer/Concert/ConcertUI/ConcertClientSharedSlate/Source/ConcertClientSharedSlate/Private/Replication/Editor/View/PropertyTree/SFilteredPropertyTreeView.h
new file mode 100644
--- /dev/null
+++ b/Engine/Plugins/Developer/Concert/ConcertUI/ConcertClientSharedSlate/Source/ConcertClientSharedSlate/Private/Replication/Editor/View/PropertyTree/SFilteredPropertyTreeView.h
@@ -0,0 +1,84 @@
+// Copyright Epic Games, Inc. All Rights Reserved.
+
+#pragma once
+
+#include "Replication/Editor/View/IPropertyTreeView.h"
+#include "Replication/Editor/View/ReplicationColumnsUtils.h"
+
+#include "Filters/SBasicFilterBar.h"
+#include "Widgets/DeclarativeSyntaxSupport.h"
+#include "Widgets/SCompoundWidget.h"
+
+namespace UE::ConcertSharedSlate
+{
+	class FReplicatedPropertyData;
+	enum class EFilterResult : uint8;
+}
+
+namespace UE::ConcertClientSharedSlate
+{
+	class SReplicationFilterBar;
+	struct FFilterablePropertyTreeViewParams;
+	
+	/** This widget extends SPropertyTreeView with filtering functionality. */
+	class SFilteredPropertyTreeView
+		: public SCompoundWidget
+		, public ConcertSharedSlate::IPropertyTreeView
+	{
+	public:
+		
+		using FFilterRef = TSharedRef<FFilterBase<const ConcertSharedSlate::FReplicatedPropertyData&>>;
+		
+		SLATE_BEGIN_ARGS(SFilteredPropertyTreeView)
+		{}
+			/*************** Arguments inherited by SReplicationTreeView ***************/
+		
+			/** The columns this list should have */
+			SLATE_ARGUMENT(TArray<ConcertSharedSlate::TReplicationColumn<ConcertSharedSlate::FReplicatedPropertyData>>, Columns)
+			/** The name of the column that will have the SExpanderArrow for the tree view. */
+			SLATE_ARGUMENT(FName, ExpandableColumnLabel)
+			/** Initial primary sort to set. */
+			SLATE_ARGUMENT(ConcertSharedSlate::FColumnSortInfo, PrimarySort)
+			/** Initial secondary sort to set. */
+			SLATE_ARGUMENT(ConcertSharedSlate::FColumnSortInfo, SecondarySort)
+		
+			/** Optional widget to add to the left of the search bar. */
+			SLATE_NAMED_SLOT(FArguments, LeftOfSearchBar)
+			/** Optional widget to add to the left of the search bar. */
+			SLATE_NAMED_SLOT(FArguments, RightOfSearchBar)
+		
+		SLATE_END_ARGS()
+
+		void Construct(const FArguments& InArgs, FFilterablePropertyTreeViewParams Params);
+		
+		//~ Begin IPropertyTreeView Interface
+		virtual void RefreshPropertyData(const TSet<FConcertPropertyChain>& PropertiesToDisplay, const FSoftClassPath& Class, bool bCanReuseExistingRowItems) override
+		{
+			ExtendedTreeView->RefreshPropertyData(PropertiesToDisplay, Class, bCanReuseExistingRowItems);
+		}
+		virtual void RequestRefilter() const override { ExtendedTreeView->RequestRefilter(); }
+		virtual void RequestResortForColumn(const FName& ColumnId) override { ExtendedTreeView->RequestResortForColumn(ColumnId); }
+		virtual TSharedRef<SWidget> GetWidget() override { return SharedThis(this); }
+		//~ End IPropertyTreeView Interface
+
+	private:
+
+		/** The tree view displaying the replicated properties */
+		TSharedPtr<IPropertyTreeView> ExtendedTreeView;
+		/** Displays the active filters*/
+		TSharedPtr<SReplicationFilterBar> FilterBar;
+
+		struct FBuildFilterBarResult
+		{
+			TArray<FFilterRef> EnabledByDefault;
+			TArray<FFilterRef> DisabledByDefault;
+		};
+		/** Build the filter bar and returns the filters that should be active by default. */
+		FBuildFilterBarResult BuildFilterBar();
+
+		/** Runs all filters through this item */
+		ConcertSharedSlate::EFilterResult PassesFilters(const ConcertSharedSlate::FReplicatedPropertyData& ReplicatedPropertyData) const;
+		bool PassesAnyFilters(const ConcertSharedSlate::FReplicatedPropertyData& ReplicatedPropertyData) const;
+	};
+}
+
diff --git a/Engine/Plugins/Developer/Concert/ConcertUI/ConcertClientSharedSlate/Source/ConcertClientSharedSlate/Public/Replication/ClientReplicationWidgetFactories.h b/Engine/Plugins/Developer/Concert/ConcertUI/ConcertClientSharedSlate/Source/ConcertClientSharedSlate/Public/Replication/ClientReplicationWidgetFactories.h
--- a/Engine/Plugins/Developer/Concert/ConcertUI/ConcertClientSharedSlate/Source/ConcertClientSharedSlate/Public/Replication/ClientReplicationWidgetFactories.h
+++ b/Engine/Plugins/Developer/Concert/ConcertUI/ConcertClientSharedSlate/Source/ConcertClientSharedSlate/Public/Replication/ClientReplicationWidgetFactories.h
@@ -2,6 +2,7 @@
 
 #pragma once
 
+#include "Replication/ReplicationWidgetFactories.h"
 #include "Templates/SharedPointer.h"
 
 struct FConcertStreamObjectAutoBindingRules;
@@ -9,6 +10,9 @@ struct FConcertObjectReplicationMap;
 
 namespace UE::ConcertSharedSlate
 {
+	struct FCreatePropertyTreeViewParams;
+	class IPropertySelectionSourceModel;
+	class IPropertyTreeView;
 	class IEditableReplicationStreamModel;
 	class IObjectNameModel;
 	class IReplicationStreamEditor;
@@ -38,9 +42,61 @@ namespace UE::ConcertClientSharedSlate
 
 	/** Simpler CreateTransactionalStreamModel overload that internally creates an UObject and sets it up automatically. */
 	CONCERTCLIENTSHAREDSLATE_API TSharedRef<ConcertSharedSlate::IEditableReplicationStreamModel> CreateTransactionalStreamModel();
+
+	/** Params for creating a filterable property tree view. */
+	struct FFilterablePropertyTreeViewParams
+	{
+		/** The columns the property view should have. The label column is always included. */
+		TArray<ConcertSharedSlate::ReplicationColumns::FReplicationPropertyColumn> PropertyColumns
+		{
+			ConcertSharedSlate::ReplicationColumns::Property::LabelColumn(),
+			ConcertSharedSlate::ReplicationColumns::Property::TypeColumn()
+		};
+		
+		/** Optional initial primary sort mode for object rows */
+		ConcertSharedSlate::FColumnSortInfo PrimaryPropertySort { ConcertSharedSlate::ReplicationColumns::Property::LabelColumnId, EColumnSortMode::Ascending };
+		/** Optional initial secondary sort mode for object rows */
+		ConcertSharedSlate::FColumnSortInfo SecondaryPropertySort { ConcertSharedSlate::ReplicationColumns::Property::LabelColumnId, EColumnSortMode::Ascending };
+	};
+	
+	/**
+	 * Creates a tree view that allows filtering of properties based on their type.
+	 * 
+	 * There is a combo box to the left of the search bar for managing the used filters.
+	 * The user can toggle used filters on and off under the search bar.
+	 */
+	CONCERTCLIENTSHAREDSLATE_API TSharedRef<ConcertSharedSlate::IPropertyTreeView> CreateFilterablePropertyTreeView(FFilterablePropertyTreeViewParams Params);
+
+	/** Params for creating the default single stream editor */
+	struct FDefaultStreamEditorParams
+	{
+		ConcertSharedSlate::FCreateEditorParams BaseEditorParams;
+
+		/**
+		 * Additional property columns to add to the default ones.
+		 * By default the default editor has the label, type, and replicates (checkbox) columns.
+		 */
+		TArray<ConcertSharedSlate::ReplicationColumns::FReplicationPropertyColumn> AdditionalPropertyColumns;
+		
+		/**
+		 * Optional. Determines the objects displayed as children to the top-level objects in the top section.
+		 * If left unspecified, the top will only display actors.
+		 * @note The created view will keep a strong reference to this.
+		 */
+		TSharedPtr<ConcertSharedSlate::IObjectHierarchyModel> ObjectHierarchy;
+
+		/** Optional. Determines the display name of objects. */
+		TSharedPtr<ConcertSharedSlate::IObjectNameModel> NameModel;
+
+		/** Called to generate the context menu for objects. This extends the options already generated by this widget and is called at the end. */
+		ConcertSharedSlate::FExtendObjectMenu OnExtendObjectsContextMenu;
+		
+		/** Additional columns to add to the object view */
+		TArray<ConcertSharedSlate::ReplicationColumns::FReplicationTopLevelObjectColumn> AdditionalObjectColumns;
+	};
 	
 	/**
-	 * Creates a default IReplicationStreamEditor.
+	 * Creates a default IReplicationStreamEditor which is used to edit a single stream, e.g. a single client's stream.
 	 *
 	 * This editor adds a checkbox to the start of every property row.
 	 * - Checking adds the property to the selected objects' property mappings.
@@ -48,5 +104,5 @@ namespace UE::ConcertClientSharedSlate
 	 * 
 	 * @see CreateViewer for a description of the UI layout.
 	 */
-	CONCERTCLIENTSHAREDSLATE_API TSharedRef<ConcertSharedSlate::IReplicationStreamEditor> CreateDefaultStreamEditor(ConcertSharedSlate::FCreateEditorParams Params);
+	CONCERTCLIENTSHAREDSLATE_API TSharedRef<ConcertSharedSlate::IReplicationStreamEditor> CreateDefaultStreamEditor(FDefaultStreamEditorParams Params);
 }
diff --git a/Engine/Plugins/Developer/Concert/ConcertUI/ConcertSharedSlate/Source/ConcertSharedSlate/Private/Replication/Editor/View/MultiEditor/SMultiReplicationStreamEditor.cpp b/Engine/Plugins/Developer/Concert/ConcertUI/ConcertSharedSlate/Source/ConcertSharedSlate/Private/Replication/Editor/View/MultiEditor/SMultiReplicationStreamEditor.cpp
--- a/Engine/Plugins/Developer/Concert/ConcertUI/ConcertSharedSlate/Source/ConcertSharedSlate/Private/Replication/Editor/View/MultiEditor/SMultiReplicationStreamEditor.cpp
+++ b/Engine/Plugins/Developer/Concert/ConcertUI/ConcertSharedSlate/Source/ConcertSharedSlate/Private/Replication/Editor/View/MultiEditor/SMultiReplicationStreamEditor.cpp
@@ -9,19 +9,18 @@
 
 namespace UE::ConcertSharedSlate
 {
-	void SMultiReplicationStreamEditor::Construct(const FArguments& InArgs, FCreateMultiStreamEditorParams Params)
+	void SMultiReplicationStreamEditor::Construct(const FArguments& InArgs, FCreateMultiStreamEditorParams EditorParams, FCreateViewerParams ViewerParams)
 	{
-		MultiStreamModel = MoveTemp(Params.MultiStreamModel);
-		ConsolidatedModel = MakeShared<FConsolidatedMultiStreamModel>(MoveTemp(Params.ConsolidatedObjectModel), MoveTemp(Params.MultiStreamModel), InArgs._GetAutoAssignStream);
+		MultiStreamModel = MoveTemp(EditorParams.MultiStreamModel);
+		ConsolidatedModel = MakeShared<FConsolidatedMultiStreamModel>(MoveTemp(EditorParams.ConsolidatedObjectModel), MoveTemp(EditorParams.MultiStreamModel), EditorParams.GetAutoAssignToStreamDelegate);
 		
-		const FCreateEditorParams BaseEditorParams
+		FCreateEditorParams BaseEditorParams
 		{
 			.DataModel = ConsolidatedModel.ToSharedRef(),
-			.ObjectSource = MoveTemp(Params.ObjectSource),
-			.PropertySource = MoveTemp(Params.PropertySource),
-			.ViewerParams = MoveTemp(Params.ViewerParams)
+			.ObjectSource = MoveTemp(EditorParams.ObjectSource),
+			.PropertySource = MoveTemp(EditorParams.PropertySource),
 		};
-		EditorView = CreateBaseStreamEditor(BaseEditorParams);
+		EditorView = CreateBaseStreamEditor(MoveTemp(BaseEditorParams), MoveTemp(ViewerParams));
 
 		ChildSlot
 		[
diff --git a/Engine/Plugins/Developer/Concert/ConcertUI/ConcertSharedSlate/Source/ConcertSharedSlate/Private/Replication/Editor/View/MultiEditor/SMultiReplicationStreamEditor.h b/Engine/Plugins/Developer/Concert/ConcertUI/ConcertSharedSlate/Source/ConcertSharedSlate/Private/Replication/Editor/View/MultiEditor/SMultiReplicationStreamEditor.h
--- a/Engine/Plugins/Developer/Concert/ConcertUI/ConcertSharedSlate/Source/ConcertSharedSlate/Private/Replication/Editor/View/MultiEditor/SMultiReplicationStreamEditor.h
+++ b/Engine/Plugins/Developer/Concert/ConcertUI/ConcertSharedSlate/Source/ConcertSharedSlate/Private/Replication/Editor/View/MultiEditor/SMultiReplicationStreamEditor.h
@@ -35,11 +35,9 @@ namespace UE::ConcertSharedSlate
 
 		SLATE_BEGIN_ARGS(SMultiReplicationStreamEditor)
 		{}
-			/** Optional. If set, the Add Actor button should automatically assign the added object to stream returned */
-			SLATE_EVENT(FGetAutoAssignTarget, GetAutoAssignStream)
 		SLATE_END_ARGS()
 
-		void Construct(const FArguments& InArgs, FCreateMultiStreamEditorParams Params);
+		void Construct(const FArguments& InArgs, FCreateMultiStreamEditorParams EditorParams, FCreateViewerParams ViewerParams);
 
 		virtual IReplicationStreamEditor& GetEditorBase() const override { return *EditorView; }
 		virtual IEditableMultiReplicationStreamModel& GetMultiStreamModel() const override { return *MultiStreamModel; }
diff --git a/Engine/Plugins/Developer/Concert/ConcertUI/ConcertSharedSlate/Source/ConcertSharedSlate/Private/Replication/Editor/View/ObjectEditor/SBaseReplicationStreamEditor.cpp b/Engine/Plugins/Developer/Concert/ConcertUI/ConcertSharedSlate/Source/ConcertSharedSlate/Private/Replication/Editor/View/ObjectEditor/SBaseReplicationStreamEditor.cpp
--- a/Engine/Plugins/Developer/Concert/ConcertUI/ConcertSharedSlate/Source/ConcertSharedSlate/Private/Replication/Editor/View/ObjectEditor/SBaseReplicationStreamEditor.cpp
+++ b/Engine/Plugins/Developer/Concert/ConcertUI/ConcertSharedSlate/Source/ConcertSharedSlate/Private/Replication/Editor/View/ObjectEditor/SBaseReplicationStreamEditor.cpp
@@ -43,12 +43,10 @@ namespace UE::ConcertSharedSlate
 		ChildSlot
 		[
 			SAssignNew(ReplicationViewer, SReplicationStreamViewer, PropertiesModelAdapter.ToSharedRef())
+				.PropertyTreeView(InArgs._PropertyTreeView)
 				.AdditionalObjectColumns(InArgs._AdditionalObjectColumns)
 				.PrimaryObjectSort(InArgs._PrimaryObjectSort)
 				.SecondaryObjectSort(InArgs._SecondaryObjectSort)
-				.AdditionalPropertyColumns(InArgs._AdditionalPropertyColumns)
-				.PrimaryPropertySort(InArgs._PrimaryPropertySort)
-				.SecondaryPropertySort(InArgs._SecondaryPropertySort)
 				.ObjectHierarchy(InArgs._ObjectHierarchy)
 				.NameModel(InArgs._NameModel)
 				.OnDeleteObjects(this, &SBaseReplicationStreamEditor::OnDeleteObjects)
@@ -67,10 +65,6 @@ namespace UE::ConcertSharedSlate
 						InArgs._LeftOfObjectSearchBar.Widget
 					]
 				]
-				.LeftOfPropertySearchBar()
-				[
-					InArgs._LeftOfPropertySearchBar.Widget
-				]
 				.NoOutlinerObjects(LOCTEXT("NoObjects", "Add objects to replicate"))
 		];
 	}
diff --git a/Engine/Plugins/Developer/Concert/ConcertUI/ConcertSharedSlate/Source/ConcertSharedSlate/Private/Replication/Editor/View/ObjectEditor/SBaseReplicationStreamEditor.h b/Engine/Plugins/Developer/Concert/ConcertUI/ConcertSharedSlate/Source/ConcertSharedSlate/Private/Replication/Editor/View/ObjectEditor/SBaseReplicationStreamEditor.h
--- a/Engine/Plugins/Developer/Concert/ConcertUI/ConcertSharedSlate/Source/ConcertSharedSlate/Private/Replication/Editor/View/ObjectEditor/SBaseReplicationStreamEditor.h
+++ b/Engine/Plugins/Developer/Concert/ConcertUI/ConcertSharedSlate/Source/ConcertSharedSlate/Private/Replication/Editor/View/ObjectEditor/SBaseReplicationStreamEditor.h
@@ -16,14 +16,14 @@ struct FConcertStreamObjectAutoBindingRules;
 
 namespace UE::ConcertSharedSlate
 {
-	class IEditableReplicationStreamModel;
 	class FFakeObjectToPropertiesEditorModel;
 	class FReplicatedObjectData;
 	class FReplicatedPropertyData;
+	class IEditableReplicationStreamModel;
 	class IObjectNameModel;
 	class IObjectSelectionSourceModel;
+	class IPropertyTreeView;
 	class IPropertySelectionSourceModel;
-	class IReplicationStreamModel;
 	class IReplicationSubobjectView;
 	class IObjectHierarchyModel;
 	class SReplicationStreamViewer;
@@ -52,6 +52,9 @@ namespace UE::ConcertSharedSlate
 		
 		SLATE_BEGIN_ARGS(SBaseReplicationStreamEditor)
 		{}
+			/** Displays the properties in a tree view */
+			SLATE_ARGUMENT(TSharedPtr<IPropertyTreeView>, PropertyTreeView)
+		
 			/** Additional columns to add to the object view */
 			SLATE_ARGUMENT(TArray<ReplicationColumns::FReplicationTopLevelObjectColumn>, AdditionalObjectColumns)
 			/** Initial primary sort to set. */
@@ -59,13 +62,6 @@ namespace UE::ConcertSharedSlate
 			/** Initial secondary sort to set. */
 			SLATE_ARGUMENT(FColumnSortInfo, SecondaryObjectSort)
 		
-			/** Additional columns to add to the property view */
-			SLATE_ARGUMENT(TArray<ReplicationColumns::FReplicationPropertyColumn>, AdditionalPropertyColumns)
-			/** Initial primary sort to set. */
-			SLATE_ARGUMENT(FColumnSortInfo, PrimaryPropertySort)
-			/** Initial secondary sort to set. */
-			SLATE_ARGUMENT(FColumnSortInfo, SecondaryPropertySort)
-		
 			/** Optional. If set, this determines the children nested under the root objects. */
 			SLATE_ARGUMENT(TSharedPtr<IObjectHierarchyModel>, ObjectHierarchy)
 			/** Optional. If set, this determines the display text for objects. */
@@ -78,11 +74,6 @@ namespace UE::ConcertSharedSlate
 			SLATE_NAMED_SLOT(FArguments, LeftOfObjectSearchBar)
 			/** Optional widget to add to the right of the object list search bar. */
 			SLATE_NAMED_SLOT(FArguments, RightOfObjectSearchBar)
-		
-			/** Optional widget to add to the left of the property list search bar. */
-			SLATE_NAMED_SLOT(FArguments, LeftOfPropertySearchBar)
-			/** Optional widget to add to the right of the property list search bar. */
-			SLATE_NAMED_SLOT(FArguments, RightOfPropertySearchBar)
 
 			/** Optional. Determines whether all UI for changing the model should be disabled. */
 			SLATE_ATTRIBUTE(bool, IsEditingEnabled)
diff --git a/Engine/Plugins/Developer/Concert/ConcertUI/ConcertSharedSlate/Source/ConcertSharedSlate/Private/Replication/Editor/View/ObjectViewer/Property/PropertyFilterBase.h b/Engine/Plugins/Developer/Concert/ConcertUI/ConcertSharedSlate/Source/ConcertSharedSlate/Private/Replication/Editor/View/ObjectViewer/Property/PropertyFilterBase.h
deleted file mode 100644
diff --git a/Engine/Plugins/Developer/Concert/ConcertUI/ConcertSharedSlate/Source/ConcertSharedSlate/Private/Replication/Editor/View/ObjectViewer/Property/PropertyFilter_ByPropertyType.cpp b/Engine/Plugins/Developer/Concert/ConcertUI/ConcertSharedSlate/Source/ConcertSharedSlate/Private/Replication/Editor/View/ObjectViewer/Property/PropertyFilter_ByPropertyType.cpp
deleted file mode 100644
diff --git a/Engine/Plugins/Developer/Concert/ConcertUI/ConcertSharedSlate/Source/ConcertSharedSlate/Private/Replication/Editor/View/ObjectViewer/Property/PropertyFilter_ByPropertyType.h b/Engine/Plugins/Developer/Concert/ConcertUI/ConcertSharedSlate/Source/ConcertSharedSlate/Private/Replication/Editor/View/ObjectViewer/Property/PropertyFilter_ByPropertyType.h
deleted file mode 100644
diff --git a/Engine/Plugins/Developer/Concert/ConcertUI/ConcertSharedSlate/Source/ConcertSharedSlate/Private/Replication/Editor/View/ObjectViewer/Property/PropertyFrontendFilter.h b/Engine/Plugins/Developer/Concert/ConcertUI/ConcertSharedSlate/Source/ConcertSharedSlate/Private/Replication/Editor/View/ObjectViewer/Property/PropertyFrontendFilter.h
deleted file mode 100644
diff --git a/Engine/Plugins/Developer/Concert/ConcertUI/ConcertSharedSlate/Source/ConcertSharedSlate/Private/Replication/Editor/View/ObjectViewer/Property/SFilteredPropertyTreeView.cpp b/Engine/Plugins/Developer/Concert/ConcertUI/ConcertSharedSlate/Source/ConcertSharedSlate/Private/Replication/Editor/View/ObjectViewer/Property/SFilteredPropertyTreeView.cpp
deleted file mode 100644
diff --git a/Engine/Plugins/Developer/Concert/ConcertUI/ConcertSharedSlate/Source/ConcertSharedSlate/Private/Replication/Editor/View/ObjectViewer/Property/SFilteredPropertyTreeView.h b/Engine/Plugins/Developer/Concert/ConcertUI/ConcertSharedSlate/Source/ConcertSharedSlate/Private/Replication/Editor/View/ObjectViewer/Property/SFilteredPropertyTreeView.h
deleted file mode 100644
diff --git a/Engine/Plugins/Developer/Concert/ConcertUI/ConcertSharedSlate/Source/ConcertSharedSlate/Private/Replication/Editor/View/ObjectViewer/Property/SPropertyTreeView.h b/Engine/Plugins/Developer/Concert/ConcertUI/ConcertSharedSlate/Source/ConcertSharedSlate/Private/Replication/Editor/View/ObjectViewer/Property/SPropertyTreeView.h
--- a/Engine/Plugins/Developer/Concert/ConcertUI/ConcertSharedSlate/Source/ConcertSharedSlate/Private/Replication/Editor/View/ObjectViewer/Property/SPropertyTreeView.h
+++ b/Engine/Plugins/Developer/Concert/ConcertUI/ConcertSharedSlate/Source/ConcertSharedSlate/Private/Replication/Editor/View/ObjectViewer/Property/SPropertyTreeView.h
@@ -3,6 +3,7 @@
 #pragma once
 
 #include "Replication/Data/ConcertPropertySelection.h"
+#include "Replication/Editor/View/IPropertyTreeView.h"
 #include "Replication/Editor/View/Tree/SReplicationTreeView.h"
 
 #include "Filters/SBasicFilterBar.h"
@@ -17,7 +18,9 @@ namespace UE::ConcertSharedSlate
 	 * This widget knows how to display a list of properties in a tree view.
 	 * It generates the items and exposes extension points for more advanced UI, such as filtering.
 	 */
-	class SPropertyTreeView : public SCompoundWidget
+	class SPropertyTreeView
+		: public SCompoundWidget
+		, public IPropertyTreeView
 	{
 	public:
 		
@@ -57,21 +60,12 @@ namespace UE::ConcertSharedSlate
 
 		void Construct(const FArguments& InArgs);
 
-		/**
-		 * Rebuilds all property data from the property source.
-		 *
-		 * @param PropertiesToDisplay The properties to display
-		 * @param Class The class from which the PropertiesToDisplay come
-		 * @param bCanReuseExistingRowItems True, will try to reuse rows for properties in the tree already (retains selected rows).
-		 *	Set this to false, if all rows should be regenerated (clears selection).
-		 *	In general, always set this to false if you've changed the object for which you're displaying the class.
-		 */
-		void RefreshPropertyData(const TSet<FConcertPropertyChain>& PropertiesToDisplay, const FSoftClassPath& Class, bool bCanReuseExistingRowItems);
-		
-		/** Called when the items need to be refiltered due to the item source changing. */
-		void RequestRefilter() const { TreeView->RequestRefilter(); }
-		/** Requests that the given column be resorted, if it currently affects the row sorting. */
-		void RequestResortForColumn(const FName& ColumnId) { TreeView->RequestResortForColumn(ColumnId); }
+		//~ Begin IPropertyTreeView Interface
+		virtual void RefreshPropertyData(const TSet<FConcertPropertyChain>& PropertiesToDisplay, const FSoftClassPath& Class, bool bCanReuseExistingRowItems) override;
+		virtual void RequestRefilter() const override { TreeView->RequestRefilter(); }
+		virtual void RequestResortForColumn(const FName& ColumnId) override { TreeView->RequestResortForColumn(ColumnId); }
+		virtual TSharedRef<SWidget> GetWidget() override { return SharedThis(this); }
+		//~ Begin IPropertyTreeView Interface
 
 	private:
 
diff --git a/Engine/Plugins/Developer/Concert/ConcertUI/ConcertSharedSlate/Source/ConcertSharedSlate/Private/Replication/Editor/View/ObjectViewer/SReplicatedPropertyView.cpp b/Engine/Plugins/Developer/Concert/ConcertUI/ConcertSharedSlate/Source/ConcertSharedSlate/Private/Replication/Editor/View/ObjectViewer/SReplicatedPropertyView.cpp
--- a/Engine/Plugins/Developer/Concert/ConcertUI/ConcertSharedSlate/Source/ConcertSharedSlate/Private/Replication/Editor/View/ObjectViewer/SReplicatedPropertyView.cpp
+++ b/Engine/Plugins/Developer/Concert/ConcertUI/ConcertSharedSlate/Source/ConcertSharedSlate/Private/Replication/Editor/View/ObjectViewer/SReplicatedPropertyView.cpp
@@ -4,10 +4,8 @@
 
 #include "Replication/Editor/Model/IReplicationStreamModel.h"
 #include "Replication/Editor/Model/ReplicatedObjectData.h"
-#include "Replication/Editor/View/SelectionViewerColumns.h"
 
 #include "Algo/AllOf.h"
-#include "Property/SFilteredPropertyTreeView.h"
 #include "Widgets/Layout/SWidgetSwitcher.h"
 #include "Widgets/Text/STextBlock.h"
 
@@ -15,8 +13,9 @@
 
 namespace UE::ConcertSharedSlate
 {
-	void SReplicatedPropertyView::Construct(const FArguments& InArgs, TSharedRef<IReplicationStreamModel> InPropertiesModel)
+	void SReplicatedPropertyView::Construct(const FArguments& InArgs, TSharedRef<IPropertyTreeView> InPropertyTreeView, TSharedRef<IReplicationStreamModel> InPropertiesModel)
 	{
+		ReplicatedProperties = MoveTemp(InPropertyTreeView);
 		PropertiesModel = MoveTemp(InPropertiesModel);
 		
 		GetSelectedRootObjectsDelegate = InArgs._GetSelectedRootObjects;
@@ -83,21 +82,6 @@ namespace UE::ConcertSharedSlate
 
 	TSharedRef<SWidget> SReplicatedPropertyView::CreatePropertiesView(const FArguments& InArgs)
 	{
-		TArray Columns
-		{
-			ReplicationColumns::Property::LabelColumn(),
-			ReplicationColumns::Property::TypeColumn()
-		};
-		Columns.Append(InArgs._AdditionalPropertyColumns);
-
-		// Set both primary and secondary in case one is overriden but always use the override.
-		const FColumnSortInfo PrimarySort = InArgs._PrimarySort.IsValid()
-			? InArgs._PrimarySort
-			: FColumnSortInfo{ ReplicationColumns::TopLevel::LabelColumnId, EColumnSortMode::Ascending };
-		const FColumnSortInfo SecondarySort = InArgs._SecondarySort.IsValid()
-			? InArgs._SecondarySort
-			: FColumnSortInfo{ ReplicationColumns::TopLevel::LabelColumnId, EColumnSortMode::Ascending };
-		
 		return SAssignNew(PropertyContent, SWidgetSwitcher)
 			// Make sure the slots are coherent with the order of EReplicatedPropertyContent!
 			.WidgetIndex(static_cast<int32>(EReplicatedPropertyContent::NoSelection))
@@ -105,14 +89,7 @@ namespace UE::ConcertSharedSlate
 			// EReplicatedPropertyContent::Properties
 			+SWidgetSwitcher::Slot()
 			[
-				SAssignNew(ReplicatedProperties, SFilteredPropertyTreeView)
-				.Columns(Columns)
-				.ExpandableColumnLabel(ReplicationColumns::Property::LabelColumnId)
-				.PrimarySort(PrimarySort)
-				.SecondarySort(SecondarySort)
-				.SelectionMode(ESelectionMode::Multi)
-				.LeftOfSearchBar() [ InArgs._LeftOfPropertySearchBar.Widget ]
-				.RightOfSearchBar() [ InArgs._RightOfPropertySearchBar.Widget ]
+				ReplicatedProperties->GetWidget()
 			]
 			
 			// EReplicatedPropertyContent::NoSelection
diff --git a/Engine/Plugins/Developer/Concert/ConcertUI/ConcertSharedSlate/Source/ConcertSharedSlate/Private/Replication/Editor/View/ObjectViewer/SReplicatedPropertyView.h b/Engine/Plugins/Developer/Concert/ConcertUI/ConcertSharedSlate/Source/ConcertSharedSlate/Private/Replication/Editor/View/ObjectViewer/SReplicatedPropertyView.h
--- a/Engine/Plugins/Developer/Concert/ConcertUI/ConcertSharedSlate/Source/ConcertSharedSlate/Private/Replication/Editor/View/ObjectViewer/SReplicatedPropertyView.h
+++ b/Engine/Plugins/Developer/Concert/ConcertUI/ConcertSharedSlate/Source/ConcertSharedSlate/Private/Replication/Editor/View/ObjectViewer/SReplicatedPropertyView.h
@@ -2,11 +2,10 @@
 
 #pragma once
 
-#include "Replication/Data/ConcertPropertySelection.h"
+#include "Replication/Editor/View/IPropertyTreeView.h"
 #include "Replication/Editor/View/SelectionViewerColumns.h"
 
 #include "Misc/Optional.h"
-#include "Property/SFilteredPropertyTreeView.h"
 #include "Widgets/DeclarativeSyntaxSupport.h"
 #include "Widgets/SCompoundWidget.h"
 
@@ -14,8 +13,6 @@ class SWidgetSwitcher;
 
 namespace UE::ConcertSharedSlate
 {
-	class SFilteredPropertyTreeView;
-
 	/**
 	 * Determines which properties are to be displayed based on an IReplicationStreamModel.
 	 * Uses a property tree for displaying. If no properties are displayed, this widget displays a message instead, e.g. to select an object.
@@ -28,26 +25,14 @@ namespace UE::ConcertSharedSlate
 
 		SLATE_BEGIN_ARGS(SReplicatedPropertyView)
 		{}
-			/** Additional columns to add to the property view */
-			SLATE_ARGUMENT(TArray<ReplicationColumns::FReplicationPropertyColumn>, AdditionalPropertyColumns)
-			/** Initial primary sort to set. */
-			SLATE_ARGUMENT(FColumnSortInfo, PrimarySort)
-			/** Initial secondary sort to set. Defaults to Label column. */
-			SLATE_ARGUMENT(FColumnSortInfo, SecondarySort)
-		
 			/** Gets the root objects selected in the object outliner. */
 			SLATE_EVENT(FGetSelectedRootObjects, GetSelectedRootObjects)
 		
 			/** Optional. If set, this determines the display text for objects. */
 			SLATE_ARGUMENT(TSharedPtr<IObjectNameModel>, NameModel)
-		
-			/** Optional widget to add to the left of the property list search bar. */
-			SLATE_NAMED_SLOT(FArguments, LeftOfPropertySearchBar)
-			/** Optional widget to add to the right of the property list search bar. */
-			SLATE_NAMED_SLOT(FArguments, RightOfPropertySearchBar)
 		SLATE_END_ARGS()
 
-		void Construct(const FArguments& InArgs, TSharedRef<IReplicationStreamModel> InPropertiesModel);
+		void Construct(const FArguments& InArgs, TSharedRef<IPropertyTreeView> InPropertyTreeView, TSharedRef<IReplicationStreamModel> InPropertiesModel);
 		
 		void RefreshPropertyData();
 		/** Requests that the given column be resorted, if it currently affects the row sorting. */
@@ -58,7 +43,7 @@ namespace UE::ConcertSharedSlate
 	private:
 		
 		/** Tree view for replicated properties. Content depends on the current object selected. */
-		TSharedPtr<SFilteredPropertyTreeView> ReplicatedProperties;
+		TSharedPtr<IPropertyTreeView> ReplicatedProperties;
 		
 		/** The model this view is visualizing. */
 		TSharedPtr<IReplicationStreamModel> PropertiesModel;
diff --git a/Engine/Plugins/Developer/Concert/ConcertUI/ConcertSharedSlate/Source/ConcertSharedSlate/Private/Replication/Editor/View/ObjectViewer/SReplicationStreamViewer.cpp b/Engine/Plugins/Developer/Concert/ConcertUI/ConcertSharedSlate/Source/ConcertSharedSlate/Private/Replication/Editor/View/ObjectViewer/SReplicationStreamViewer.cpp
--- a/Engine/Plugins/Developer/Concert/ConcertUI/ConcertSharedSlate/Source/ConcertSharedSlate/Private/Replication/Editor/View/ObjectViewer/SReplicationStreamViewer.cpp
+++ b/Engine/Plugins/Developer/Concert/ConcertUI/ConcertSharedSlate/Source/ConcertSharedSlate/Private/Replication/Editor/View/ObjectViewer/SReplicationStreamViewer.cpp
@@ -267,14 +267,9 @@ namespace UE::ConcertSharedSlate
 				]
 				.BodyContent()
 				[
-					SAssignNew(PropertySection, SReplicatedPropertyView, PropertiesModel.ToSharedRef())
-					.AdditionalPropertyColumns(InArgs._AdditionalPropertyColumns)
-					.PrimarySort(InArgs._PrimaryPropertySort)
-					.SecondarySort(InArgs._SecondaryPropertySort)
+					SAssignNew(PropertySection, SReplicatedPropertyView, InArgs._PropertyTreeView.ToSharedRef(), PropertiesModel.ToSharedRef())
 					.GetSelectedRootObjects_Lambda([this](){ return GetSelectedOutlinerObjects(); })
 					.NameModel(InArgs._NameModel)
-					.LeftOfPropertySearchBar() [ InArgs._LeftOfPropertySearchBar.Widget ]
-					.RightOfPropertySearchBar() [ InArgs._RightOfPropertySearchBar.Widget ]
 				]
 			];
 	}
diff --git a/Engine/Plugins/Developer/Concert/ConcertUI/ConcertSharedSlate/Source/ConcertSharedSlate/Private/Replication/Editor/View/ObjectViewer/SReplicationStreamViewer.h b/Engine/Plugins/Developer/Concert/ConcertUI/ConcertSharedSlate/Source/ConcertSharedSlate/Private/Replication/Editor/View/ObjectViewer/SReplicationStreamViewer.h
--- a/Engine/Plugins/Developer/Concert/ConcertUI/ConcertSharedSlate/Source/ConcertSharedSlate/Private/Replication/Editor/View/ObjectViewer/SReplicationStreamViewer.h
+++ b/Engine/Plugins/Developer/Concert/ConcertUI/ConcertSharedSlate/Source/ConcertSharedSlate/Private/Replication/Editor/View/ObjectViewer/SReplicationStreamViewer.h
@@ -49,19 +49,15 @@ namespace UE::ConcertSharedSlate
 
 		SLATE_BEGIN_ARGS(SReplicationStreamViewer)
 		{}
+			/** Displays the properties in a tree view */
+			SLATE_ARGUMENT(TSharedPtr<IPropertyTreeView>, PropertyTreeView)
+		
 			/** Additional columns to add to the object view */
 			SLATE_ARGUMENT(TArray<ReplicationColumns::FReplicationTopLevelObjectColumn>, AdditionalObjectColumns)
 			/** Initial primary sort to set. */
 			SLATE_ARGUMENT(FColumnSortInfo, PrimaryObjectSort)
 			/** Initial secondary sort to set. */
 			SLATE_ARGUMENT(FColumnSortInfo, SecondaryObjectSort)
-		
-			/** Additional columns to add to the property view */
-			SLATE_ARGUMENT(TArray<ReplicationColumns::FReplicationPropertyColumn>, AdditionalPropertyColumns)
-			/** Initial primary sort to set. */
-			SLATE_ARGUMENT(FColumnSortInfo, PrimaryPropertySort)
-			/** Initial secondary sort to set. */
-			SLATE_ARGUMENT(FColumnSortInfo, SecondaryPropertySort)
 
 			/** Optional. If set, this determines the children nested under the root objects. */
 			SLATE_ARGUMENT(TSharedPtr<IObjectHierarchyModel>, ObjectHierarchy)
@@ -78,11 +74,6 @@ namespace UE::ConcertSharedSlate
 			SLATE_NAMED_SLOT(FArguments, LeftOfObjectSearchBar)
 			/** Optional widget to add to the right of the object list search bar. */
 			SLATE_NAMED_SLOT(FArguments, RightOfObjectSearchBar)
-		
-			/** Optional widget to add to the left of the property list search bar. */
-			SLATE_NAMED_SLOT(FArguments, LeftOfPropertySearchBar)
-			/** Optional widget to add to the right of the object list search bar. */
-			SLATE_NAMED_SLOT(FArguments, RightOfPropertySearchBar)
 
 			/** Optional text to display when no object is in the outliner. Defaults to "No objects to display." "*/
 			SLATE_ATTRIBUTE(FText, NoOutlinerObjects)
diff --git a/Engine/Plugins/Developer/Concert/ConcertUI/ConcertSharedSlate/Source/ConcertSharedSlate/Private/Replication/Editor/View/Tree/SReplicationTreeView.h b/Engine/Plugins/Developer/Concert/ConcertUI/ConcertSharedSlate/Source/ConcertSharedSlate/Private/Replication/Editor/View/Tree/SReplicationTreeView.h
--- a/Engine/Plugins/Developer/Concert/ConcertUI/ConcertSharedSlate/Source/ConcertSharedSlate/Private/Replication/Editor/View/Tree/SReplicationTreeView.h
+++ b/Engine/Plugins/Developer/Concert/ConcertUI/ConcertSharedSlate/Source/ConcertSharedSlate/Private/Replication/Editor/View/Tree/SReplicationTreeView.h
@@ -220,7 +220,7 @@ namespace UE::ConcertSharedSlate
 		void SetPrimarySortMode(FName SortedColumnId, EColumnSortMode::Type SortMode)
 		{
 			const TReplicationColumn<TItemType>* Column = FindColumnByName(SortedColumnId);
-			if (ensure(Column && Column->CanBeSorted()))
+			if (ensureAlways(Column && Column->CanBeSorted()))
 			{
 				PrimarySortInfo = { SortedColumnId, SortMode };
 				
diff --git a/Engine/Plugins/Developer/Concert/ConcertUI/ConcertSharedSlate/Source/ConcertSharedSlate/Private/Replication/ReplicationWidgetFactories.cpp b/Engine/Plugins/Developer/Concert/ConcertUI/ConcertSharedSlate/Source/ConcertSharedSlate/Private/Replication/ReplicationWidgetFactories.cpp
--- a/Engine/Plugins/Developer/Concert/ConcertUI/ConcertSharedSlate/Source/ConcertSharedSlate/Private/Replication/ReplicationWidgetFactories.cpp
+++ b/Engine/Plugins/Developer/Concert/ConcertUI/ConcertSharedSlate/Source/ConcertSharedSlate/Private/Replication/ReplicationWidgetFactories.cpp
@@ -5,6 +5,7 @@
 #include "Editor/Model/GenericReplicationStreamModel.h"
 #include "Editor/View/MultiEditor/SMultiReplicationStreamEditor.h"
 #include "Editor/View/ObjectEditor/SBaseReplicationStreamEditor.h"
+#include "Editor/View/ObjectViewer/Property/SPropertyTreeView.h"
 
 namespace UE::ConcertSharedSlate
 {
@@ -16,30 +17,57 @@ namespace UE::ConcertSharedSlate
 		return MakeShared<FGenericReplicationStreamModel>(MoveTemp(ReplicationMapAttribute), MoveTemp(Extender));
 	}
 	
-	TSharedRef<IReplicationStreamEditor> CreateBaseStreamEditor(FCreateEditorParams Params)
+	TSharedRef<IReplicationStreamEditor> CreateBaseStreamEditor(FCreateEditorParams EditorParams, FCreateViewerParams ViewerParams)
 	{
-		return SNew(SBaseReplicationStreamEditor, Params.DataModel, Params.ObjectSource, Params.PropertySource)
-			.AdditionalObjectColumns(Params.ViewerParams.AdditionalObjectColumns)
-			.PrimaryObjectSort(Params.ViewerParams.PrimaryObjectSort)
-			.SecondaryObjectSort(Params.ViewerParams.SecondaryObjectSort)
-			.AdditionalPropertyColumns(Params.ViewerParams.AdditionalPropertyColumns)
-			.PrimaryPropertySort(Params.ViewerParams.PrimaryPropertySort)
-			.SecondaryPropertySort(Params.ViewerParams.SecondaryPropertySort)
-			.ObjectHierarchy(Params.ViewerParams.ObjectHierarchy)
-			.NameModel(Params.ViewerParams.NameModel)
-			.OnExtendObjectsContextMenu(Params.ViewerParams.OnExtendObjectsContextMenu)
-			.LeftOfObjectSearchBar() [ Params.ViewerParams.LeftOfObjectSearchBar.Widget ]
-			.RightOfObjectSearchBar() [ Params.ViewerParams.RightOfObjectSearchBar.Widget ]
-			.LeftOfPropertySearchBar() [ Params.ViewerParams.LeftOfPropertySearchBar.Widget ]
-			.RightOfPropertySearchBar() [ Params.ViewerParams.RightOfPropertySearchBar.Widget ]
-			.IsEditingEnabled(Params.IsEditingEnabled)
-			.EditingDisabledToolTipText(Params.EditingDisabledToolTipText);
+		return SNew(SBaseReplicationStreamEditor, MoveTemp(EditorParams.DataModel), MoveTemp(EditorParams.ObjectSource), MoveTemp(EditorParams.PropertySource))
+			.PropertyTreeView(MoveTemp(ViewerParams.PropertyTreeView))
+			.AdditionalObjectColumns(MoveTemp(ViewerParams.AdditionalObjectColumns))
+			.PrimaryObjectSort(ViewerParams.PrimaryObjectSort)
+			.SecondaryObjectSort(ViewerParams.SecondaryObjectSort)
+			.ObjectHierarchy(MoveTemp(ViewerParams.ObjectHierarchy))
+			.NameModel(MoveTemp(ViewerParams.NameModel))
+			.OnExtendObjectsContextMenu(MoveTemp(ViewerParams.OnExtendObjectsContextMenu))
+			.LeftOfObjectSearchBar() [ MoveTemp(ViewerParams.LeftOfObjectSearchBar.Widget) ]
+			.RightOfObjectSearchBar() [ MoveTemp(ViewerParams.RightOfObjectSearchBar.Widget) ]
+			.IsEditingEnabled(MoveTemp(EditorParams.IsEditingEnabled))
+			.EditingDisabledToolTipText(MoveTemp(EditorParams.EditingDisabledToolTipText));
+	}
+	
+	TSharedRef<IPropertyTreeView> CreateSearchablePropertyTreeView(FCreatePropertyTreeViewParams Params)
+	{
+		// The label columns is always required
+		const bool bHasLabel = Params.PropertyColumns.ContainsByPredicate([](const ReplicationColumns::FReplicationPropertyColumn& Column)
+		{
+			return Column.ColumnId == ReplicationColumns::Property::LabelColumnId;
+		});
+		if (!bHasLabel)
+		{
+			Params.PropertyColumns.Add(ReplicationColumns::Property::LabelColumn());
+		}
+
+		SReplicationTreeView<FReplicatedPropertyData>::FCustomFilter FilterDelegate = Params.FilterItem.IsBound()
+			? SReplicationTreeView<FReplicatedPropertyData>::FCustomFilter::CreateLambda([Filter = MoveTemp(Params.FilterItem)](const TSharedPtr<FReplicatedPropertyData>& Item)
+			{
+				return Filter.Execute(*Item.Get()) == EFilterResult::PassesFilter;
+			})
+			: SReplicationTreeView<FReplicatedPropertyData>::FCustomFilter{};
+		
+		return SNew(SPropertyTreeView)
+			.FilterItem(MoveTemp(FilterDelegate))
+			.Columns(MoveTemp(Params.PropertyColumns))
+			.ExpandableColumnLabel(ReplicationColumns::Property::LabelColumnId)
+			.PrimarySort(Params.PrimaryPropertySort)
+			.SecondarySort(Params.SecondaryPropertySort)
+			.SelectionMode(ESelectionMode::Multi)
+			.LeftOfSearchBar() [ MoveTemp(Params.LeftOfPropertySearchBar.Widget) ]
+			.RightOfSearchBar() [ MoveTemp(Params.RightOfPropertySearchBar.Widget) ]
+			.RowBelowSearchBar() [ MoveTemp(Params.RowBelowSearchBar.Widget) ]
+			.NoItemsContent() [ MoveTemp(Params.NoItemsContent.Widget) ];
 	}
 
-	TSharedRef<IMultiReplicationStreamEditor> CreateBaseMultiStreamEditor(FCreateMultiStreamEditorParams Params)
+	TSharedRef<IMultiReplicationStreamEditor> CreateBaseMultiStreamEditor(FCreateMultiStreamEditorParams EditorParams, FCreateViewerParams ViewerParams)
 	{
-		return SNew(SMultiReplicationStreamEditor, MoveTemp(Params))
-			.GetAutoAssignStream(MoveTemp(Params.GetAutoAssignToStreamDelegate));
+		return SNew(SMultiReplicationStreamEditor, MoveTemp(EditorParams), MoveTemp(ViewerParams));
 	}
 }
 
diff --git a/Engine/Plugins/Developer/Concert/ConcertUI/ConcertSharedSlate/Source/ConcertSharedSlate/Public/Replication/Editor/View/IPropertyTreeView.h b/Engine/Plugins/Developer/Concert/ConcertUI/ConcertSharedSlate/Source/ConcertSharedSlate/Public/Replication/Editor/View/IPropertyTreeView.h
new file mode 100644
--- /dev/null
+++ b/Engine/Plugins/Developer/Concert/ConcertUI/ConcertSharedSlate/Source/ConcertSharedSlate/Public/Replication/Editor/View/IPropertyTreeView.h
@@ -0,0 +1,47 @@
+// Copyright Epic Games, Inc. All Rights Reserved.
+
+#pragma once
+
+#include "Containers/Set.h"
+#include "Templates/SharedPointer.h"
+
+class SWidget;
+
+struct FConcertPropertyChain;
+struct FSoftClassPath;
+
+namespace UE::ConcertSharedSlate
+{
+	/** Represents a tree view displaying properties. */
+	class CONCERTSHAREDSLATE_API IPropertyTreeView
+	{
+	public:
+		
+		/**
+		 * Rebuilds all property data from the property source.
+		 *
+		 * @param PropertiesToDisplay The properties to display
+		 * @param Class The class from which the PropertiesToDisplay come
+		 * @param bCanReuseExistingRowItems True, will try to reuse rows for properties in the tree already (retains selected rows).
+		 *	Set this to false, if all rows should be regenerated (clears selection).
+		 *	In general, always set this to false if you've changed the object for which you're displaying the class.
+		 */
+		virtual void RefreshPropertyData(const TSet<FConcertPropertyChain>& PropertiesToDisplay, const FSoftClassPath& Class, bool bCanReuseExistingRowItems) = 0;
+		
+		/**
+		 * Reapply the filter function to all items at the end of the frame. Call e.g. when the filters have changed.
+		 */
+		virtual void RequestRefilter() const = 0;
+		
+		/**
+		 * Requests that the given column be resorted, if it currently affects the row sorting (primary or secondary).
+		 * Call e.g. when a sortable attribute of the column has changed.
+		 */
+		virtual void RequestResortForColumn(const FName& ColumnId) = 0;
+
+		/** Gets the tree view's widget */
+		virtual TSharedRef<SWidget> GetWidget() = 0;
+
+		virtual ~IPropertyTreeView() = default;
+	};
+}
diff --git a/Engine/Plugins/Developer/Concert/ConcertUI/ConcertSharedSlate/Source/ConcertSharedSlate/Public/Replication/ReplicationWidgetFactories.h b/Engine/Plugins/Developer/Concert/ConcertUI/ConcertSharedSlate/Source/ConcertSharedSlate/Public/Replication/ReplicationWidgetFactories.h
--- a/Engine/Plugins/Developer/Concert/ConcertUI/ConcertSharedSlate/Source/ConcertSharedSlate/Public/Replication/ReplicationWidgetFactories.h
+++ b/Engine/Plugins/Developer/Concert/ConcertUI/ConcertSharedSlate/Source/ConcertSharedSlate/Public/Replication/ReplicationWidgetFactories.h
@@ -7,6 +7,7 @@
 
 #include "Delegates/Delegate.h"
 #include "Editor/View/ReplicationColumnsUtils.h"
+#include "Editor/View/SelectionViewerColumns.h"
 #include "Misc/Attribute.h"
 #include "Templates/SharedPointer.h"
 
@@ -18,13 +19,15 @@ struct FConcertObjectReplicationMap;
 
 namespace UE::ConcertSharedSlate
 {
+	class FReplicatedPropertyData;
+	
 	class IEditableMultiReplicationStreamModel;
 	class IEditableReplicationStreamModel;
 	class IMultiReplicationStreamEditor;
-	
-	class IReplicationStreamModel;
 	class IObjectNameModel;
 	class IObjectSelectionSourceModel;
+	class IPropertyTreeView;
+	class IReplicationStreamModel;
 	class IReplicationStreamEditor;
 	class IReplicationStreamViewer;
 	class IReplicationSubobjectView;
@@ -44,10 +47,55 @@ namespace UE::ConcertSharedSlate
 		TAttribute<FConcertObjectReplicationMap*> ReplicationMapAttribute,
 		TSharedPtr<IStreamExtender> Extender = nullptr
 		);
+
+	enum class EFilterResult : uint8
+	{
+		/** Include the object into list of displayed objects */
+		PassesFilter,
+		/** Exclude the object from list of displayed objects */
+		DoesNotPassFilter
+	};
+	DECLARE_DELEGATE_RetVal_OneParam(EFilterResult, FFilterPropertyData, const FReplicatedPropertyData&);
+	
+	struct FCreatePropertyTreeViewParams
+	{
+		/** Optional, additional columns to add to the property view. The LabelColumn will always be added. */
+		TArray<ReplicationColumns::FReplicationPropertyColumn> PropertyColumns
+		{
+			ReplicationColumns::Property::LabelColumn(),
+			ReplicationColumns::Property::TypeColumn()
+		};
+
+		/** Optional filter function. Return true to al */
+		FFilterPropertyData FilterItem;
+		
+		/** Optional initial primary sort mode for object rows */
+		FColumnSortInfo PrimaryPropertySort { ReplicationColumns::Property::LabelColumnId, EColumnSortMode::Ascending };
+		/** Optional initial secondary sort mode for object rows */
+		FColumnSortInfo SecondaryPropertySort { ReplicationColumns::Property::LabelColumnId, EColumnSortMode::Ascending };
+		
+		/** Optional widget to add to the left of the property list search bar. */
+		TAlwaysValidWidget LeftOfPropertySearchBar;
+		/** Optional widget to add to the right of the property list search bar. */
+		TAlwaysValidWidget RightOfPropertySearchBar;
+		/** Optional widget to add between the search bar and the table view (e.g. a SBasicFilterBar). */
+		TAlwaysValidWidget RowBelowSearchBar;
+		/** Optional, alternate content to show instead of the tree view when there are no rows. */
+		TAlwaysValidWidget NoItemsContent;
+	};
+	
+	/**
+	 * Creates a tree view that uses a search box for filtering items.
+	 * You can customize this tree view by adding custom widgets and columns into the property view.
+	 */
+	CONCERTSHAREDSLATE_API TSharedRef<IPropertyTreeView> CreateSearchablePropertyTreeView(FCreatePropertyTreeViewParams Params = {});
 	
 	/** Params for creating a IReplicationStreamViewer. */
 	struct FCreateViewerParams
 	{
+		/** Required. Displays the properties in a tree view. You can pass in e.g. custom UI with advanced filtering. */
+		TSharedRef<IPropertyTreeView> PropertyTreeView = CreateSearchablePropertyTreeView();
+		
 		/**
 		 * Optional. Determines the objects displayed as children to the top-level objects in the top section.
 		 * If left unspecified, the top will only display actors.
@@ -68,21 +116,10 @@ namespace UE::ConcertSharedSlate
 		/** Optional initial secondary sort mode for object rows */
 		FColumnSortInfo SecondaryObjectSort;
 		
-		/** Additional columns to add to the property view */
-		TArray<ReplicationColumns::FReplicationPropertyColumn> AdditionalPropertyColumns;
-		/** Optional initial primary sort mode for object rows */
-		FColumnSortInfo PrimaryPropertySort;
-		/** Optional initial secondary sort mode for object rows */
-		FColumnSortInfo SecondaryPropertySort;
-		
 		/** Optional widget to add to the left of the object list search bar. */
 		TAlwaysValidWidget LeftOfObjectSearchBar;
 		/** Optional widget to add to the right of the object list search bar. */
 		TAlwaysValidWidget RightOfObjectSearchBar;
-		/** Optional widget to add to the left of the property list search bar. */
-		TAlwaysValidWidget LeftOfPropertySearchBar;
-		/** Optional widget to add to the right of the property list search bar. */
-		TAlwaysValidWidget RightOfPropertySearchBar;
 	};
 
 	// TODO DP 5.5: Create factory function that uses FCreateViewerParams and creates an IReplicationStreamViewer
@@ -91,18 +128,18 @@ namespace UE::ConcertSharedSlate
 	struct FCreateEditorParams
 	{
 		/**
-		 * The model that the editor is displaying.
+		 * Required. The model that the editor is displaying.
 		 * @note The view will keep a strong reference to this.
 		 */
 		TSharedRef<IEditableReplicationStreamModel> DataModel;
 		
 		/**
-		 * Determines the objects that can be added to the object list. 
+		 * Required. Determines the objects that can be added to the object list. 
 		 * @note The view will keep a strong reference to this.
 		 */
 		TSharedRef<IObjectSelectionSourceModel> ObjectSource;
 		/**
-		 * Determines the objects that are displayed in the property list. 
+		 * Required. Determines the properties that are displayed in the property list. 
 		 * @note The view will keep a strong reference to this.
 		 */
 		TSharedRef<IPropertySelectionSourceModel> PropertySource;
@@ -111,9 +148,6 @@ namespace UE::ConcertSharedSlate
 		TAttribute<bool> IsEditingEnabled;
 		/** Optional. Whenever IsEditingEnabled returns true, this tooltip is displayed for relevant, disabled UI. */
 		TAttribute<FText> EditingDisabledToolTipText;
-
-		/** Base params for customizing the viewing part of the editor */
-		FCreateViewerParams ViewerParams;
 	};
 
 	/**
@@ -127,14 +161,9 @@ namespace UE::ConcertSharedSlate
 	 * 1. Root objects, similar to world outliner:
 	 *		- Similar to World Outliner: Displays top-level objects, i.e. actors, are added here.
 	 *		- FCreateEditorParams::ObjectSource is used to build a combo button through which new objects can be added.
-	 * 2. Subobjects (optional):
-	 *		- Similar to component view (SSubobjectEditor)
-	 *		- Shows subobjects of a root objects selected above; typically components.
-	 * 3. Properties:
-	 *		- Similar to details panel
-	 *		- Shows properties of the selected root object and / or subobjects.
+	 * 2. Properties: Shows properties of the selected root object and / or subobjects.
 	 */
-	CONCERTSHAREDSLATE_API TSharedRef<IReplicationStreamEditor> CreateBaseStreamEditor(FCreateEditorParams Params);
+	CONCERTSHAREDSLATE_API TSharedRef<IReplicationStreamEditor> CreateBaseStreamEditor(FCreateEditorParams EditorParams, FCreateViewerParams ViewerParams);
 
 	DECLARE_DELEGATE_RetVal_OneParam(TSharedPtr<IEditableReplicationStreamModel>, FGetAutoAssignTarget, TConstArrayView<UObject*>);
 	/** Params for creating an IMultiReplicationStreamEditor */
@@ -175,9 +204,6 @@ namespace UE::ConcertSharedSlate
 
 		/** Optional. If set, the Add Actor button should automatically assign the added object to stream returned */
 		FGetAutoAssignTarget GetAutoAssignToStreamDelegate;
-		
-		/** Base params for customizing the viewing part of the editor */
-		FCreateViewerParams ViewerParams;
 	};
 
 	/**
@@ -189,5 +215,5 @@ namespace UE::ConcertSharedSlate
 	 *
 	 * @note This editor does NOT create any UI for changing the assigned properties. You are supposed to inject column widgets to the property rows for this.
 	 */
-	CONCERTSHAREDSLATE_API TSharedRef<IMultiReplicationStreamEditor> CreateBaseMultiStreamEditor(FCreateMultiStreamEditorParams Params);
+	CONCERTSHAREDSLATE_API TSharedRef<IMultiReplicationStreamEditor> CreateBaseMultiStreamEditor(FCreateMultiStreamEditorParams EditorParams, FCreateViewerParams ViewerParams);
 }
