Fix forgetting to unregister global delegate FCoreUObjectDelegates::OnObjectTransacted

diff --git a/Engine/Plugins/Developer/Concert/ConcertUI/ConcertClientSharedSlate/Source/ConcertClientSharedSlate/Private/Replication/Editor/UnrealEditor/ModifyObjectInLevelHandler.cpp b/Engine/Plugins/Developer/Concert/ConcertUI/ConcertClientSharedSlate/Source/ConcertClientSharedSlate/Private/Replication/Editor/UnrealEditor/ModifyObjectInLevelHandler.cpp
--- a/Engine/Plugins/Developer/Concert/ConcertUI/ConcertClientSharedSlate/Source/ConcertClientSharedSlate/Private/Replication/Editor/UnrealEditor/ModifyObjectInLevelHandler.cpp
+++ b/Engine/Plugins/Developer/Concert/ConcertUI/ConcertClientSharedSlate/Source/ConcertClientSharedSlate/Private/Replication/Editor/UnrealEditor/ModifyObjectInLevelHandler.cpp
@@ -30,6 +30,7 @@ namespace UE::ConcertClientSharedSlate
 		{
 			GEngine->OnLevelActorDeleted().RemoveAll(this);
 		}
+		FCoreUObjectDelegates::OnObjectTransacted.RemoveAll(this);
 	}
 
 	void FModifyObjectInLevelHandler::OnActorDeleted(AActor* Actor) const
