Replace IPropertySoureModel with new IPropertySource.
IPropertySource strips away functions that were used in an old version but are no longer used in 5.5.

diff --git a/Engine/Plugins/Developer/Concert/ConcertUI/ConcertClientSharedSlate/Source/ConcertClientSharedSlate/Private/Replication/Editor/Model/PropertySource/ConcertSyncCoreReplicatedPropertySource.cpp b/Engine/Plugins/Developer/Concert/ConcertUI/ConcertClientSharedSlate/Source/ConcertClientSharedSlate/Private/Replication/Editor/Model/PropertySource/ConcertSyncCoreReplicatedPropertySource.cpp
--- a/Engine/Plugins/Developer/Concert/ConcertUI/ConcertClientSharedSlate/Source/ConcertClientSharedSlate/Private/Replication/Editor/Model/PropertySource/ConcertSyncCoreReplicatedPropertySource.cpp
+++ b/Engine/Plugins/Developer/Concert/ConcertUI/ConcertClientSharedSlate/Source/ConcertClientSharedSlate/Private/Replication/Editor/Model/PropertySource/ConcertSyncCoreReplicatedPropertySource.cpp
@@ -38,6 +38,7 @@ namespace UE::ConcertClientSharedSlate
 		};
 	}
 
+	PRAGMA_DISABLE_DEPRECATION_WARNINGS
 	void FConcertSyncCoreReplicatedPropertySource::EnumerateSelectableItems(TFunctionRef<EBreakBehavior(const ConcertSharedSlate::FSelectablePropertyInfo& SelectableOption)> Delegate) const
 	{
 		if (!Class.IsValid())
@@ -50,6 +51,7 @@ namespace UE::ConcertClientSharedSlate
 			return Delegate({ MoveTemp(Property) });
 		});
 	}
+	PRAGMA_ENABLE_DEPRECATION_WARNINGS
 }
 
 #undef LOCTEXT_NAMESPACE
\ No newline at end of file
diff --git a/Engine/Plugins/Developer/Concert/ConcertUI/ConcertClientSharedSlate/Source/ConcertClientSharedSlate/Private/Replication/Editor/Model/PropertySource/SelectPropertyFromUClassModel.cpp b/Engine/Plugins/Developer/Concert/ConcertUI/ConcertClientSharedSlate/Source/ConcertClientSharedSlate/Private/Replication/Editor/Model/PropertySource/SelectPropertyFromUClassModel.cpp
--- a/Engine/Plugins/Developer/Concert/ConcertUI/ConcertClientSharedSlate/Source/ConcertClientSharedSlate/Private/Replication/Editor/Model/PropertySource/SelectPropertyFromUClassModel.cpp
+++ b/Engine/Plugins/Developer/Concert/ConcertUI/ConcertClientSharedSlate/Source/ConcertClientSharedSlate/Private/Replication/Editor/Model/PropertySource/SelectPropertyFromUClassModel.cpp
@@ -2,27 +2,20 @@
 
 #include "Replication/Editor/Model/PropertySource/SelectPropertyFromUClassModel.h"
 
-#include "Replication/Editor/Model/Property/IPropertySourceModel.h"
-#include "Replication/Editor/Model/PropertySource/ConcertSyncCoreReplicatedPropertySource.h"
-
 #include "ConcertLogGlobal.h"
-
-#define LOCTEXT_NAMESPACE "FSelectPropertyFromUClassModel"
+#include "Replication/Editor/Model/PropertySource/ReplicatablePropertySource.h"
 
 namespace UE::ConcertClientSharedSlate
 {
 	void FSelectPropertyFromUClassModel::ProcessPropertySource(
 		const ConcertSharedSlate::FPropertySourceContext& Context,
-		TFunctionRef<void(const ConcertSharedSlate::IPropertySourceModel& Model)> Processor
+		TFunctionRef<void(const ConcertSharedSlate::IPropertySource& Model)> Processor
 		) const
 	{
 		UClass* LoadedClass = Context.Class.TryLoadClass<UObject>();
-		
-		FConcertSyncCoreReplicatedPropertySource UClassIteratorSource;
-		UClassIteratorSource.SetClass(LoadedClass ? LoadedClass : nullptr);
 		UE_CLOG(!LoadedClass, LogConcert, Warning, TEXT("Could not resolve class %s. Properties will not be available."), *Context.Class.ToString());
-		Processor(UClassIteratorSource);
+		
+		FReplicatablePropertySource PropertySource(LoadedClass);
+		Processor(PropertySource);
 	}
-}
-
-#undef LOCTEXT_NAMESPACE
\ No newline at end of file
+}
\ No newline at end of file
diff --git a/Engine/Plugins/Developer/Concert/ConcertUI/ConcertClientSharedSlate/Source/ConcertClientSharedSlate/Public/Replication/Editor/Model/PropertySource/ConcertSyncCoreReplicatedPropertySource.h b/Engine/Plugins/Developer/Concert/ConcertUI/ConcertClientSharedSlate/Source/ConcertClientSharedSlate/Public/Replication/Editor/Model/PropertySource/ConcertSyncCoreReplicatedPropertySource.h
--- a/Engine/Plugins/Developer/Concert/ConcertUI/ConcertClientSharedSlate/Source/ConcertClientSharedSlate/Public/Replication/Editor/Model/PropertySource/ConcertSyncCoreReplicatedPropertySource.h
+++ b/Engine/Plugins/Developer/Concert/ConcertUI/ConcertClientSharedSlate/Source/ConcertClientSharedSlate/Public/Replication/Editor/Model/PropertySource/ConcertSyncCoreReplicatedPropertySource.h
@@ -6,6 +6,7 @@
 #include "Replication/Editor/Model/Property/IPropertySourceModel.h"
 
 #include "HAL/Platform.h"
+#include "Misc/CoreMiscDefines.h"
 #include "UObject/WeakObjectPtr.h"
 
 class UClass;
@@ -13,7 +14,10 @@ class UClass;
 namespace UE::ConcertClientSharedSlate
 {
 	/** The allowed properties are those returned by UE::ConcertSyncCore::ForEachReplicatableProperty.*/
-	class CONCERTCLIENTSHAREDSLATE_API FConcertSyncCoreReplicatedPropertySource : public ConcertSharedSlate::IPropertySourceModel
+	class CONCERTCLIENTSHAREDSLATE_API UE_DEPRECATED(5.5, "No longer in use") FConcertSyncCoreReplicatedPropertySource
+		PRAGMA_DISABLE_DEPRECATION_WARNINGS
+		: public ConcertSharedSlate::IPropertySourceModel
+		PRAGMA_ENABLE_DEPRECATION_WARNINGS
 	{
 	public:
 
@@ -22,7 +26,9 @@ namespace UE::ConcertClientSharedSlate
 		//~ Begin IPropertySourceModel Interface
 		virtual ConcertSharedSlate::FSourceDisplayInfo GetDisplayInfo() const override;
 		virtual uint32 GetNumSelectableItems() const override { return NumProperties; }
+		PRAGMA_DISABLE_DEPRECATION_WARNINGS
 		virtual void EnumerateSelectableItems(TFunctionRef<EBreakBehavior(const ConcertSharedSlate::FSelectablePropertyInfo& SelectableOption)> Delegate) const override;
+		PRAGMA_ENABLE_DEPRECATION_WARNINGS
 		//~ End IPropertySourceModel Interface
 
 	private:
diff --git a/Engine/Plugins/Developer/Concert/ConcertUI/ConcertClientSharedSlate/Source/ConcertClientSharedSlate/Public/Replication/Editor/Model/PropertySource/ReplicatablePropertySource.h b/Engine/Plugins/Developer/Concert/ConcertUI/ConcertClientSharedSlate/Source/ConcertClientSharedSlate/Public/Replication/Editor/Model/PropertySource/ReplicatablePropertySource.h
new file mode 100644
--- /dev/null
+++ b/Engine/Plugins/Developer/Concert/ConcertUI/ConcertClientSharedSlate/Source/ConcertClientSharedSlate/Public/Replication/Editor/Model/PropertySource/ReplicatablePropertySource.h
@@ -0,0 +1,38 @@
+﻿// Copyright Epic Games, Inc. All Rights Reserved.
+
+#pragma once
+
+#include "Replication/PropertyChainUtils.h"
+#include "Replication/Editor/Model/Property/IPropertySource.h"
+
+namespace UE::ConcertClientSharedSlate
+{
+	/** Lists all properties that can be replicated. */
+	class FReplicatablePropertySource : public ConcertSharedSlate::IPropertySource
+	{
+	public:
+
+		explicit FReplicatablePropertySource(UClass* Class)
+			: Class(Class)
+		{}
+
+		//~ Begin IPropertySource Interface
+		virtual void EnumerateProperties(TFunctionRef<EBreakBehavior(ConcertSharedSlate::FPropertyInfo&& Property)> Delegate) const override
+		{
+			if (!Class.IsValid())
+			{
+				return;
+			}
+			
+			ConcertSyncCore::PropertyChain::ForEachReplicatableConcertProperty(*Class.Get(), [this, &Delegate](FConcertPropertyChain&& Property)
+			{
+				return Delegate(ConcertSharedSlate::FPropertyInfo(MoveTemp(Property)));
+			});
+		}
+		//~ End IPropertySource Interface
+
+	private:
+		
+		TWeakObjectPtr<UClass> Class;
+	};
+}
diff --git a/Engine/Plugins/Developer/Concert/ConcertUI/ConcertClientSharedSlate/Source/ConcertClientSharedSlate/Public/Replication/Editor/Model/PropertySource/SelectPropertyFromUClassModel.h b/Engine/Plugins/Developer/Concert/ConcertUI/ConcertClientSharedSlate/Source/ConcertClientSharedSlate/Public/Replication/Editor/Model/PropertySource/SelectPropertyFromUClassModel.h
--- a/Engine/Plugins/Developer/Concert/ConcertUI/ConcertClientSharedSlate/Source/ConcertClientSharedSlate/Public/Replication/Editor/Model/PropertySource/SelectPropertyFromUClassModel.h
+++ b/Engine/Plugins/Developer/Concert/ConcertUI/ConcertClientSharedSlate/Source/ConcertClientSharedSlate/Public/Replication/Editor/Model/PropertySource/SelectPropertyFromUClassModel.h
@@ -10,7 +10,6 @@ struct FSoftClassPath;
 
 namespace UE::ConcertClientSharedSlate
 {
-	
 	/**
 	 * Decides which properties can be added to IEditableReplicationStreamModel.
 	 * The allowed properties are those returned by UE::ConcertSyncCore::ForEachReplicatableProperty.
@@ -22,10 +21,9 @@ namespace UE::ConcertClientSharedSlate
 		//~ Begin IPropertySelectionProcessor Interface
 		virtual void ProcessPropertySource(
 			const ConcertSharedSlate::FPropertySourceContext& Context,
-			TFunctionRef<void(const ConcertSharedSlate::IPropertySourceModel& Model)> Processor
+			TFunctionRef<void(const ConcertSharedSlate::IPropertySource& Model)> Processor
 			) const override;
 		//~ End IPropertySelectionProcessor Interface
-
 	};
 }
 
diff --git a/Engine/Plugins/Developer/Concert/ConcertUI/ConcertSharedSlate/Source/ConcertSharedSlate/Private/Replication/Editor/Utils/PropertyEnumerationUtils.cpp b/Engine/Plugins/Developer/Concert/ConcertUI/ConcertSharedSlate/Source/ConcertSharedSlate/Private/Replication/Editor/Utils/PropertyEnumerationUtils.cpp
--- a/Engine/Plugins/Developer/Concert/ConcertUI/ConcertSharedSlate/Source/ConcertSharedSlate/Private/Replication/Editor/Utils/PropertyEnumerationUtils.cpp
+++ b/Engine/Plugins/Developer/Concert/ConcertUI/ConcertSharedSlate/Source/ConcertSharedSlate/Private/Replication/Editor/Utils/PropertyEnumerationUtils.cpp
@@ -4,6 +4,7 @@
 
 #include "Misc/EBreakBehavior.h"
 #include "Replication/Editor/Model/IReplicationStreamModel.h"
+#include "Replication/Editor/Model/Property/IPropertySource.h"
 #include "Replication/Editor/Model/Property/IPropertySourceProcessor.h"
 
 #include "Containers/Set.h"
@@ -73,9 +74,9 @@ namespace UE::ConcertSharedSlate
 			const FPropertySourceContext ObjectQueryContext(Object, ObjectClass);
 					
 			EBreakBehavior BreakBehavior = EBreakBehavior::Continue;
-			Source.ProcessPropertySource(ObjectQueryContext, [&Callback, &ObjectClass, &BreakBehavior](const IPropertySourceModel& PropertySource)
+			Source.ProcessPropertySource(ObjectQueryContext, [&Callback, &ObjectClass, &BreakBehavior](const IPropertySource& PropertySource)
 			{
-				PropertySource.EnumerateSelectableItems([&Callback, &ObjectClass, &BreakBehavior](const FSelectablePropertyInfo& PropertyInfo)
+				PropertySource.EnumerateProperties([&Callback, &ObjectClass, &BreakBehavior](const FPropertyInfo& PropertyInfo)
 				{
 					BreakBehavior = Callback(ObjectClass, PropertyInfo.Property);
 					return BreakBehavior;
diff --git a/Engine/Plugins/Developer/Concert/ConcertUI/ConcertSharedSlate/Source/ConcertSharedSlate/Public/Replication/Editor/Model/Property/IPropertySelectionSourceModel.h b/Engine/Plugins/Developer/Concert/ConcertUI/ConcertSharedSlate/Source/ConcertSharedSlate/Public/Replication/Editor/Model/Property/IPropertySelectionSourceModel.h
--- a/Engine/Plugins/Developer/Concert/ConcertUI/ConcertSharedSlate/Source/ConcertSharedSlate/Public/Replication/Editor/Model/Property/IPropertySelectionSourceModel.h
+++ b/Engine/Plugins/Developer/Concert/ConcertUI/ConcertSharedSlate/Source/ConcertSharedSlate/Public/Replication/Editor/Model/Property/IPropertySelectionSourceModel.h
@@ -12,7 +12,10 @@
 
 namespace UE::ConcertSharedSlate
 {
-	using FPropertySourceCategory UE_DEPRECATED(5.5, "No longer in use.") = ConcertSharedSlate::TSourceSelectionCategory<FSelectablePropertyInfo>;
+	using FPropertySourceCategory UE_DEPRECATED(5.5, "No longer in use.") =
+		PRAGMA_DISABLE_DEPRECATION_WARNINGS
+		ConcertSharedSlate::TSourceSelectionCategory<FSelectablePropertyInfo>;
+		PRAGMA_ENABLE_DEPRECATION_WARNINGS
 	
 	/** Decides which properties can be added to a IReplicationStreamModel. */
 	class UE_DEPRECATED(5.5, "Use IPropertySourceProcessor instead.") IPropertySelectionSourceModel
@@ -23,7 +26,9 @@ namespace UE::ConcertSharedSlate
 	public:
 
 		/** Gets the single source determining which properties can be selected. */
+		PRAGMA_DISABLE_DEPRECATION_WARNINGS
 		virtual TSharedRef<IPropertySourceModel> GetPropertySource(const FSoftClassPath& Class) const = 0;
+		PRAGMA_ENABLE_DEPRECATION_WARNINGS
 		
 		virtual ~IPropertySelectionSourceModel() = default;
 	};
diff --git a/Engine/Plugins/Developer/Concert/ConcertUI/ConcertSharedSlate/Source/ConcertSharedSlate/Public/Replication/Editor/Model/Property/IPropertySource.h b/Engine/Plugins/Developer/Concert/ConcertUI/ConcertSharedSlate/Source/ConcertSharedSlate/Public/Replication/Editor/Model/Property/IPropertySource.h
new file mode 100644
--- /dev/null
+++ b/Engine/Plugins/Developer/Concert/ConcertUI/ConcertSharedSlate/Source/ConcertSharedSlate/Public/Replication/Editor/Model/Property/IPropertySource.h
@@ -0,0 +1,31 @@
+﻿// Copyright Epic Games, Inc. All Rights Reserved.
+
+#pragma once
+
+#include "Misc/EBreakBehavior.h"
+#include "Replication/Data/ConcertPropertySelection.h"
+
+#include "Templates/FunctionFwd.h"
+#include "UObject/SoftObjectPath.h"
+#include "UObject/SoftObjectPtr.h"
+
+namespace UE::ConcertSharedSlate
+{
+	/** Wraps FConcertPropertyChain so it is easier to potentially change what IPropertySource lists in the future. */
+	struct FPropertyInfo
+	{
+		FConcertPropertyChain Property;
+		explicit FPropertyInfo(FConcertPropertyChain Property) : Property(MoveTemp(Property)) {}
+	};
+	
+	/** Lists out a bunch of properties. */
+	class IPropertySource
+	{
+	public:
+
+		/** Lists a bunch of properties. */
+		virtual void EnumerateProperties(TFunctionRef<EBreakBehavior(FPropertyInfo&& Property)> Delegate) const = 0;
+		
+		virtual ~IPropertySource() = default;
+	};
+}
diff --git a/Engine/Plugins/Developer/Concert/ConcertUI/ConcertSharedSlate/Source/ConcertSharedSlate/Public/Replication/Editor/Model/Property/IPropertySourceModel.h b/Engine/Plugins/Developer/Concert/ConcertUI/ConcertSharedSlate/Source/ConcertSharedSlate/Public/Replication/Editor/Model/Property/IPropertySourceModel.h
--- a/Engine/Plugins/Developer/Concert/ConcertUI/ConcertSharedSlate/Source/ConcertSharedSlate/Public/Replication/Editor/Model/Property/IPropertySourceModel.h
+++ b/Engine/Plugins/Developer/Concert/ConcertUI/ConcertSharedSlate/Source/ConcertSharedSlate/Public/Replication/Editor/Model/Property/IPropertySourceModel.h
@@ -5,10 +5,12 @@
 #include "Model/Item/IItemSourceModel.h"
 #include "Replication/Data/ConcertPropertySelection.h"
 
+#include "Misc/CoreMiscDefines.h"
+
 namespace UE::ConcertSharedSlate
 {
 	/** Info about a selectable property */
-	struct FSelectablePropertyInfo
+	struct UE_DEPRECATED(5.5, "No longer in use.") FSelectablePropertyInfo
 	{
 		FConcertPropertyChain Property;
 	};
@@ -17,5 +19,8 @@ namespace UE::ConcertSharedSlate
 	 * A specific object source, e.g. like actors, components from an actor (right-click), etc.
 	 * @see IPropertySelectionProcessor.
 	 */
-	using IPropertySourceModel = ConcertSharedSlate::IItemSourceModel<FSelectablePropertyInfo>;
+	using IPropertySourceModel UE_DEPRECATED(5.5, "No longer in use.") =
+		PRAGMA_DISABLE_DEPRECATION_WARNINGS
+		ConcertSharedSlate::IItemSourceModel<FSelectablePropertyInfo>;
+		PRAGMA_ENABLE_DEPRECATION_WARNINGS
 }
diff --git a/Engine/Plugins/Developer/Concert/ConcertUI/ConcertSharedSlate/Source/ConcertSharedSlate/Public/Replication/Editor/Model/Property/IPropertySourceProcessor.h b/Engine/Plugins/Developer/Concert/ConcertUI/ConcertSharedSlate/Source/ConcertSharedSlate/Public/Replication/Editor/Model/Property/IPropertySourceProcessor.h
--- a/Engine/Plugins/Developer/Concert/ConcertUI/ConcertSharedSlate/Source/ConcertSharedSlate/Public/Replication/Editor/Model/Property/IPropertySourceProcessor.h
+++ b/Engine/Plugins/Developer/Concert/ConcertUI/ConcertSharedSlate/Source/ConcertSharedSlate/Public/Replication/Editor/Model/Property/IPropertySourceProcessor.h
@@ -2,14 +2,14 @@
 
 #pragma once
 
-#include "Replication/Editor/Model/Property/IPropertySourceModel.h"
-
 #include "Templates/FunctionFwd.h"
 #include "UObject/SoftObjectPath.h"
 #include "UObject/SoftObjectPtr.h"
 
 namespace UE::ConcertSharedSlate
 {
+	class IPropertySource;
+	
 	/** Arguments for getting properties associated with an object / class. */
 	struct FPropertySourceContext
 	{
@@ -35,7 +35,7 @@ namespace UE::ConcertSharedSlate
 	public:
 		
 		/** Allows you to get the properties that are associated with the passed in object / class context. You are not allowed to keep a reference to Processor. */
-		virtual void ProcessPropertySource(const FPropertySourceContext& Context, TFunctionRef<void(const IPropertySourceModel& Model)> Processor) const = 0;
+		virtual void ProcessPropertySource(const FPropertySourceContext& Context, TFunctionRef<void(const IPropertySource& Model)> Processor) const = 0;
 		
 		virtual ~IPropertySourceProcessor() = default;
 	};
