Extract object display name into global function

diff --git a/Engine/Plugins/Developer/Concert/ConcertUI/ConcertClientSharedSlate/Source/ConcertClientSharedSlate/Private/Replication/Editor/Model/Object/EditorObjectNameModel.cpp b/Engine/Plugins/Developer/Concert/ConcertUI/ConcertClientSharedSlate/Source/ConcertClientSharedSlate/Private/Replication/Editor/Model/Object/EditorObjectNameModel.cpp
deleted file mode 100644
diff --git a/Engine/Plugins/Developer/Concert/ConcertUI/ConcertClientSharedSlate/Source/ConcertClientSharedSlate/Private/Replication/Editor/Model/Object/EditorObjectNameModel.h b/Engine/Plugins/Developer/Concert/ConcertUI/ConcertClientSharedSlate/Source/ConcertClientSharedSlate/Private/Replication/Editor/Model/Object/EditorObjectNameModel.h
--- a/Engine/Plugins/Developer/Concert/ConcertUI/ConcertClientSharedSlate/Source/ConcertClientSharedSlate/Private/Replication/Editor/Model/Object/EditorObjectNameModel.h
+++ b/Engine/Plugins/Developer/Concert/ConcertUI/ConcertClientSharedSlate/Source/ConcertClientSharedSlate/Private/Replication/Editor/Model/Object/EditorObjectNameModel.h
@@ -2,6 +2,7 @@
 
 #pragma once
 
+#include "Replication/ObjectNameUtils.h"
 #include "Replication/Editor/Model/Object/IObjectNameModel.h"
 
 namespace UE::ConcertClientSharedSlate
@@ -10,9 +11,9 @@ namespace UE::ConcertClientSharedSlate
 	class FEditorObjectNameModel : public ConcertSharedSlate::IObjectNameModel
 	{
 	public:
-
+		
 		//~ Begin IObjectNameModel Interface
-		virtual FText GetObjectDisplayName(const TSoftObjectPtr<>& ObjectPath) const override;
+		virtual FText GetObjectDisplayName(const TSoftObjectPtr<>& ObjectPath) const override { return ConcertClientSharedSlate::GetObjectDisplayName(ObjectPath); }
 		//~ End IObjectNameModel Interface
 	};
 }
diff --git a/Engine/Plugins/Developer/Concert/ConcertUI/ConcertClientSharedSlate/Source/ConcertClientSharedSlate/Private/Replication/ObjectNameUtils.cpp b/Engine/Plugins/Developer/Concert/ConcertUI/ConcertClientSharedSlate/Source/ConcertClientSharedSlate/Private/Replication/ObjectNameUtils.cpp
new file mode 100644
--- /dev/null
+++ b/Engine/Plugins/Developer/Concert/ConcertUI/ConcertClientSharedSlate/Source/ConcertClientSharedSlate/Private/Replication/ObjectNameUtils.cpp
@@ -0,0 +1,59 @@
+﻿// Copyright Epic Games, Inc. All Rights Reserved.
+
+#include "Replication/ObjectNameUtils.h"
+
+#include "Internationalization/Text.h"
+#include "Replication/Utils/ObjectUtils.h"
+
+#include "GameFramework/Actor.h"
+#include "SubobjectDataSubsystem.h"
+
+namespace UE::ConcertClientSharedSlate
+{
+	namespace Private
+	{
+		static FText FindSubobjectDisplayName(const UObject& Subbject, AActor& OwningActor)
+		{
+			USubobjectDataSubsystem* SubobjectDataSubsystem = USubobjectDataSubsystem::Get();
+			TArray<FSubobjectDataHandle> Handles;
+			SubobjectDataSubsystem->GatherSubobjectData(&OwningActor, Handles);
+				
+			for (const FSubobjectDataHandle& Handle : Handles)
+			{
+				const FSubobjectData* SubobjectData = Handle.GetData();
+				const UObject* Object = SubobjectData->FindComponentInstanceInActor(&OwningActor);
+				if (Object == &Subbject)
+				{
+					constexpr bool bShowNativeComponentNames = false;
+					return FText::FromString(SubobjectData->GetDisplayString(bShowNativeComponentNames));
+				}
+			}
+
+			return FText::GetEmpty();
+		}
+	}
+
+	FText GetObjectDisplayName(const TSoftObjectPtr<>& ObjectPath)
+	{
+		if (UObject* ResolvedObject = ObjectPath.Get())
+		{
+			// Display actor just like the outliner does
+			if (const AActor* AsActor = Cast<AActor>(ResolvedObject))
+			{
+				return FText::FromString(AsActor->GetActorLabel());
+			}
+
+			// Display the same component name as the SSubobjectEditor widget does, i.e. component hierarchy in the details panel or Blueprint editor.
+			AActor* OwningActor = ResolvedObject->GetTypedOuter<AActor>();
+			const FText FoundSubobjectName = OwningActor ? Private::FindSubobjectDisplayName(*ResolvedObject, *OwningActor) : FText::GetEmpty();
+			if (!FoundSubobjectName.IsEmpty())
+			{
+				return FoundSubobjectName;
+			}
+
+			return FText::FromString(ResolvedObject->GetName());
+		}
+		
+		return FText::FromString(ConcertSharedSlate::ObjectUtils::ExtractObjectDisplayStringFromPath(ObjectPath.GetUniqueID()));
+	}
+}
diff --git a/Engine/Plugins/Developer/Concert/ConcertUI/ConcertClientSharedSlate/Source/ConcertClientSharedSlate/Public/Replication/ObjectNameUtils.h b/Engine/Plugins/Developer/Concert/ConcertUI/ConcertClientSharedSlate/Source/ConcertClientSharedSlate/Public/Replication/ObjectNameUtils.h
new file mode 100644
--- /dev/null
+++ b/Engine/Plugins/Developer/Concert/ConcertUI/ConcertClientSharedSlate/Source/ConcertClientSharedSlate/Public/Replication/ObjectNameUtils.h
@@ -0,0 +1,14 @@
+﻿// Copyright Epic Games, Inc. All Rights Reserved.
+
+#pragma once
+
+#include "UObject/SoftObjectPtr.h"
+
+class FText;
+
+namespace UE::ConcertClientSharedSlate
+{
+	/** @return The display name of ObjectPath to use for Concert replication UI. */
+	CONCERTCLIENTSHAREDSLATE_API FText GetObjectDisplayName(const TSoftObjectPtr<>& ObjectPath);
+}
+
