Fix: Replication can't add Sources from a Sub Level

diff --git a/Engine/Plugins/Developer/Concert/ConcertUI/ConcertClientSharedSlate/Source/ConcertClientSharedSlate/Private/Replication/Editor/UnrealEditor/HideObjectsNotInWorldLogic.cpp b/Engine/Plugins/Developer/Concert/ConcertUI/ConcertClientSharedSlate/Source/ConcertClientSharedSlate/Private/Replication/Editor/UnrealEditor/HideObjectsNotInWorldLogic.cpp
--- a/Engine/Plugins/Developer/Concert/ConcertUI/ConcertClientSharedSlate/Source/ConcertClientSharedSlate/Private/Replication/Editor/UnrealEditor/HideObjectsNotInWorldLogic.cpp
+++ b/Engine/Plugins/Developer/Concert/ConcertUI/ConcertClientSharedSlate/Source/ConcertClientSharedSlate/Private/Replication/Editor/UnrealEditor/HideObjectsNotInWorldLogic.cpp
@@ -3,6 +3,7 @@
 #include "Replication/Editor/UnrealEditor/HideObjectsNotInWorldLogic.h"
 
 #include "Engine/Engine.h"
+#include "Engine/Level.h"
 #include "Engine/World.h"
 #include "UObject/SoftObjectPath.h"
 
@@ -33,8 +34,22 @@ namespace UE::ConcertClientSharedSlate
 			return false;
 		}
 
-		// If it does not resolve, it's no in GWorld (the world should load all the actors in it).
 		const UObject* Object = ObjectPath.ResolveObject();
-		return Object && Object->IsIn(GWorld);
+		if (!Object)
+		{
+			// If it does not resolve, it's no in GWorld (the world should load all the actors in it).
+			return false;
+		}
+
+		if (Object->IsIn(GWorld))
+		{
+			return true;
+		}
+
+		const bool bIsInSublevel = GWorld->GetLevels().ContainsByPredicate([&Object](ULevel* Level)
+		{
+			return Level && Object->IsIn(Level);
+		});
+		return bIsInSublevel;
 	}
 }
