Fix Blueprints not having their properties auto-added when added for replication.
The bug was that objects that had been trashed previously were being added instead of the live instances.

diff --git a/Engine/Plugins/Developer/Concert/ConcertUI/ConcertClientSharedSlate/Source/ConcertClientSharedSlate/Private/Replication/Settings/ConcertSubobjectMatchingRules.cpp b/Engine/Plugins/Developer/Concert/ConcertUI/ConcertClientSharedSlate/Source/ConcertClientSharedSlate/Private/Replication/Settings/ConcertSubobjectMatchingRules.cpp
--- a/Engine/Plugins/Developer/Concert/ConcertUI/ConcertClientSharedSlate/Source/ConcertClientSharedSlate/Private/Replication/Settings/ConcertSubobjectMatchingRules.cpp
+++ b/Engine/Plugins/Developer/Concert/ConcertUI/ConcertClientSharedSlate/Source/ConcertClientSharedSlate/Private/Replication/Settings/ConcertSubobjectMatchingRules.cpp
@@ -6,6 +6,8 @@
 #include "Internationalization/Regex.h"
 #include "Misc/EBreakBehavior.h"
 
+#include "UObject/ObjectMacros.h"
+
 namespace UE::ConcertSharedSlate
 {
 	static bool MatchesAnyRegex(const FString& Input, const TSet<FString>& AllRegex)
@@ -64,5 +66,5 @@ void FConcertSubobjectMatchingRules::MatchToSubobjectsBreakable(const UObject& A
 		}
 		
 		return true;
-	}, bIncludeNested);
+	}, bIncludeNested, RF_MirroredGarbage);
 }
\ No newline at end of file
