Replication outliner now has a SDropTarget allowing users to drag-drop actors.

diff --git a/Engine/Plugins/Developer/Concert/ConcertApp/MultiUserClient/Source/MultiUserClient/Private/Widgets/ActiveSession/Replication/Client/Multi/SMultiClientView.cpp b/Engine/Plugins/Developer/Concert/ConcertApp/MultiUserClient/Source/MultiUserClient/Private/Widgets/ActiveSession/Replication/Client/Multi/SMultiClientView.cpp
--- a/Engine/Plugins/Developer/Concert/ConcertApp/MultiUserClient/Source/MultiUserClient/Private/Widgets/ActiveSession/Replication/Client/Multi/SMultiClientView.cpp
+++ b/Engine/Plugins/Developer/Concert/ConcertApp/MultiUserClient/Source/MultiUserClient/Private/Widgets/ActiveSession/Replication/Client/Multi/SMultiClientView.cpp
@@ -171,7 +171,10 @@ namespace UE::MultiUserClient::Replication
 					)
 			},
 			.ShouldDisplayObjectDelegate = FShouldDisplayObject::CreateSP(this, &SMultiClientView::ShouldDisplayObject),
-			.MakeObjectRowOverlayWidgetDelegate = FMakeObjectRowOverlayWidget::CreateSP(this, &SMultiClientView::MakeObjectRowOverlayWidget)
+			.MakeObjectRowOverlayWidgetDelegate = FMakeObjectRowOverlayWidget::CreateSP(this, &SMultiClientView::MakeObjectRowOverlayWidget),
+			.WrapOutlinerWidgetDelegate = ConcertClientSharedSlate::CreateDropTargetOutlinerWrapper(
+				{ ConcertClientSharedSlate::FDragDropReplicatableObject::CreateSP(this, &SMultiClientView::HandleDroppedObjects) }
+				)
 		};
 		ViewerParams.RightOfObjectSearchBar.Widget = CreateRightOfSearchBarContent(InConcertClient, InMultiUserReplicationManager);
 		StreamEditor = CreateBaseMultiStreamEditor(MoveTemp(Params), MoveTemp(ViewerParams));
@@ -294,14 +297,32 @@ namespace UE::MultiUserClient::Replication
 			: SNullWidget::NullWidget;
 	}
 
-	void SMultiClientView::OnPreAddObjectsFromComboButton(TArrayView<const ConcertSharedSlate::FSelectableObjectInfo>) const
+	void SMultiClientView::HandleDroppedObjects(TConstArrayView<UObject*> DroppedObjects) const
+	{
+		EnableObjectExtensionOnAdd();
+
+		ConcertSharedSlate::IEditableReplicationStreamModel& Model = StreamEditor->GetConsolidatedModel();
+		TArray<UObject*> ObjectsToAdd; 
+		Algo::TransformIf(DroppedObjects, ObjectsToAdd,
+			[&Model](const TWeakObjectPtr<>& WeakObject)
+			{
+				UObject* Object = WeakObject.Get();
+				return Object && !Model.ContainsObjects({ Object });
+			},
+			[](const TWeakObjectPtr<>& WeakObject){ return WeakObject.Get(); }
+			);
+		Model.AddObjects(ObjectsToAdd);
+
+		DisableObjectExtensionOnAdd();
+	}
+
+	void SMultiClientView::EnableObjectExtensionOnAdd() const
 	{
-		// When the user adds using the combo button, automatically add discover relevant objects and properties
 		OnlineClientManager->GetLocalClient().GetStreamExtender()
 			.SetShouldExtend(true);
 	}
 
-	void SMultiClientView::OnPostAddObjectsFromComboButton(TArrayView<const ConcertSharedSlate::FSelectableObjectInfo>) const
+	void SMultiClientView::DisableObjectExtensionOnAdd() const
 	{
 		OnlineClientManager->GetLocalClient().GetStreamExtender()
 			.SetShouldExtend(false);
diff --git a/Engine/Plugins/Developer/Concert/ConcertApp/MultiUserClient/Source/MultiUserClient/Private/Widgets/ActiveSession/Replication/Client/Multi/SMultiClientView.h b/Engine/Plugins/Developer/Concert/ConcertApp/MultiUserClient/Source/MultiUserClient/Private/Widgets/ActiveSession/Replication/Client/Multi/SMultiClientView.h
--- a/Engine/Plugins/Developer/Concert/ConcertApp/MultiUserClient/Source/MultiUserClient/Private/Widgets/ActiveSession/Replication/Client/Multi/SMultiClientView.h
+++ b/Engine/Plugins/Developer/Concert/ConcertApp/MultiUserClient/Source/MultiUserClient/Private/Widgets/ActiveSession/Replication/Client/Multi/SMultiClientView.h
@@ -107,7 +107,16 @@ namespace UE::MultiUserClient::Replication
 		/** Creates the widget that overlays actor rows. */
 		TSharedRef<SWidget> MakeObjectRowOverlayWidget(const ConcertSharedSlate::FReplicatedObjectData& ReplicatedObjectData) const;
 		
-		void OnPreAddObjectsFromComboButton(TArrayView<const ConcertSharedSlate::FSelectableObjectInfo>) const;
-		void OnPostAddObjectsFromComboButton(TArrayView<const ConcertSharedSlate::FSelectableObjectInfo>) const;
+		/** When the user adds using the combo button, automatically add discover relevant objects and properties. */
+		void OnPreAddObjectsFromComboButton(TArrayView<const ConcertSharedSlate::FSelectableObjectInfo>) const { EnableObjectExtensionOnAdd(); }
+		void OnPostAddObjectsFromComboButton(TArrayView<const ConcertSharedSlate::FSelectableObjectInfo>) const { DisableObjectExtensionOnAdd(); }
+		
+		/** Called when objects are dropped into the view. */
+		void HandleDroppedObjects(TConstArrayView<UObject*> DroppedObjects) const;
+
+		/** Enables adding common properties and subobjects when an object is added to replication. */
+		void EnableObjectExtensionOnAdd() const;
+		/** Stops adding common properties and subobjects when an object is added to replication. */
+		void DisableObjectExtensionOnAdd() const;
 	};
 }
diff --git a/Engine/Plugins/Developer/Concert/ConcertUI/ConcertClientSharedSlate/Source/ConcertClientSharedSlate/ConcertClientSharedSlate.Build.cs b/Engine/Plugins/Developer/Concert/ConcertUI/ConcertClientSharedSlate/Source/ConcertClientSharedSlate/ConcertClientSharedSlate.Build.cs
--- a/Engine/Plugins/Developer/Concert/ConcertUI/ConcertClientSharedSlate/Source/ConcertClientSharedSlate/ConcertClientSharedSlate.Build.cs
+++ b/Engine/Plugins/Developer/Concert/ConcertUI/ConcertClientSharedSlate/Source/ConcertClientSharedSlate/ConcertClientSharedSlate.Build.cs
@@ -28,7 +28,8 @@ namespace UnrealBuildTool.Rules
 				{
 					// Engine
 					"AssetDefinition",
-					"AssetRegistry",
+					"AssetRegistry", 
+					"EditorWidgets",
 					"InputCore",
 					"Projects",
 					"Slate",
diff --git a/Engine/Plugins/Developer/Concert/ConcertUI/ConcertClientSharedSlate/Source/ConcertClientSharedSlate/Private/Replication/ClientReplicationWidgetFactories.cpp b/Engine/Plugins/Developer/Concert/ConcertUI/ConcertClientSharedSlate/Source/ConcertClientSharedSlate/Private/Replication/ClientReplicationWidgetFactories.cpp
--- a/Engine/Plugins/Developer/Concert/ConcertUI/ConcertClientSharedSlate/Source/ConcertClientSharedSlate/Private/Replication/ClientReplicationWidgetFactories.cpp
+++ b/Engine/Plugins/Developer/Concert/ConcertUI/ConcertClientSharedSlate/Source/ConcertClientSharedSlate/Private/Replication/ClientReplicationWidgetFactories.cpp
@@ -2,6 +2,7 @@
 
 #include "Replication/ClientReplicationWidgetFactories.h"
 
+#include "Editor/View/SReplicationDropArea.h"
 #include "Replication/ReplicationWidgetFactories.h"
 #include "Replication/Editor/Model/Object/EditorObjectHierarchyModel.h"
 #include "Replication/Editor/Model/Object/EditorObjectNameModel.h"
@@ -56,4 +57,16 @@ namespace UE::ConcertClientSharedSlate
 	{
 		return SNew(SFilteredPropertyTreeView, MoveTemp(Params));
 	}
-}
\ No newline at end of file
+
+	ConcertSharedSlate::FWrapOutlinerWidget CreateDropTargetOutlinerWrapper(FCreateDropTargetOutlinerWrapperParams Params)
+	{
+		return ConcertSharedSlate::FWrapOutlinerWidget::CreateLambda([Params = MoveTemp(Params)](const TSharedRef<SWidget>& Widget)
+		{
+			return SNew(SReplicationDropArea)
+				.HandleDroppedObjects(Params.HandleDroppedObjectsDelegate)
+				[
+					Widget
+				];
+		});
+	}
+}
diff --git a/Engine/Plugins/Developer/Concert/ConcertUI/ConcertClientSharedSlate/Source/ConcertClientSharedSlate/Private/Replication/Editor/View/SReplicationDropArea.cpp b/Engine/Plugins/Developer/Concert/ConcertUI/ConcertClientSharedSlate/Source/ConcertClientSharedSlate/Private/Replication/Editor/View/SReplicationDropArea.cpp
new file mode 100644
--- /dev/null
+++ b/Engine/Plugins/Developer/Concert/ConcertUI/ConcertClientSharedSlate/Source/ConcertClientSharedSlate/Private/Replication/Editor/View/SReplicationDropArea.cpp
@@ -0,0 +1,129 @@
+﻿// Copyright Epic Games, Inc. All Rights Reserved.
+
+#include "SReplicationDropArea.h"
+
+#include "Misc/EBreakBehavior.h"
+
+#include "EditorActorFolders.h"
+#include "SDropTarget.h"
+#include "DragAndDrop/ActorDragDropOp.h"
+#include "DragAndDrop/CompositeDragDropOp.h"
+#include "DragAndDrop/FolderDragDropOp.h"
+
+#include <type_traits>
+
+namespace UE::ConcertClientSharedSlate
+{
+	template<typename TCallback> requires std::is_invocable_r_v<EBreakBehavior, TCallback, UObject*>
+	static void ForEachObjectInOperation(const TSharedPtr<FDragDropOperation>& InOperation, TCallback&& Callback)
+	{
+		if (!InOperation || !InOperation.IsValid())
+		{
+			return;
+		}
+		
+		TSharedPtr<FActorDragDropOp>  ActorDrag = nullptr;
+		TSharedPtr<FFolderDragDropOp> FolderDrag = nullptr;
+
+		if (InOperation->IsOfType<FActorDragDropOp>())
+		{
+			ActorDrag = StaticCastSharedPtr<FActorDragDropOp>(InOperation);
+		}
+		else if (InOperation->IsOfType<FFolderDragDropOp>())
+		{
+			FolderDrag = StaticCastSharedPtr<FFolderDragDropOp>(InOperation);
+		}
+		else if (InOperation->IsOfType<FCompositeDragDropOp>())
+		{
+			if (const TSharedPtr<FCompositeDragDropOp> CompositeDrag = StaticCastSharedPtr<FCompositeDragDropOp>(InOperation))
+			{
+				ActorDrag = CompositeDrag->GetSubOp<FActorDragDropOp>();
+				FolderDrag = CompositeDrag->GetSubOp<FFolderDragDropOp>();
+			}
+		}
+
+		if (ActorDrag)
+		{
+			for (TWeakObjectPtr<AActor> WeakActor : ActorDrag->Actors)
+			{
+				if (AActor* Actor = WeakActor.Get())
+				{
+					if (Callback(Actor) == EBreakBehavior::Break)
+					{
+						return;
+					}
+				}
+			}
+		}
+
+		if (FolderDrag)
+		{
+			TArray<AActor*> FolderActors;
+			FActorFolders::GetActorsFromFolders(*GWorld, FolderDrag->Folders, FolderActors);
+
+			for (AActor* ActorInFolder : FolderActors)
+			{
+				if (Callback(ActorInFolder) == EBreakBehavior::Break)
+				{
+					return;
+				}
+			}
+		}
+	}
+}
+
+namespace UE::ConcertClientSharedSlate
+{
+	void SReplicationDropArea::Construct(const FArguments& InArgs)
+	{
+		HandleDroppedObjectsDelegate = InArgs._HandleDroppedObjects;
+		CanDropObjectDelegate = InArgs._CanDropObject;
+		
+		ChildSlot
+		[
+			SNew(SDropTarget)
+			.OnDropped(this, &SReplicationDropArea::OnDragDropTarget)
+			.OnAllowDrop(this, &SReplicationDropArea::CanDragDropTarget)
+			.OnIsRecognized(this, &SReplicationDropArea::CanDragDropTarget)
+			[
+				InArgs._Content.Widget
+			]
+		];
+	}
+
+	FReply SReplicationDropArea::OnDragDropTarget(const FGeometry& Geometry, const FDragDropEvent& DragDropEvent) const
+	{
+		TArray<UObject*> AllowedDroppedObjects;
+		ForEachObjectInOperation(DragDropEvent.GetOperation(), [this, &AllowedDroppedObjects](UObject* Object)
+		{
+			if (CanDrop(Object))
+			{
+				AllowedDroppedObjects.Add(Object);
+			}
+			return EBreakBehavior::Continue;
+		});
+
+		if (ensure(!AllowedDroppedObjects.IsEmpty()))
+		{
+			HandleDroppedObjectsDelegate.Execute(AllowedDroppedObjects);
+		}
+		
+		return FReply::Handled();
+	}
+
+	bool SReplicationDropArea::CanDragDropTarget(TSharedPtr<FDragDropOperation> InOperation) const
+	{
+		bool bCanHandle = false;
+		ForEachObjectInOperation(InOperation, [this, &bCanHandle](UObject* Object)
+		{
+			bCanHandle = CanDrop(Object);
+			return bCanHandle ? EBreakBehavior::Break : EBreakBehavior::Continue;
+		});
+		return bCanHandle;
+	}
+
+	bool SReplicationDropArea::CanDrop(UObject* Object) const
+	{
+		return Object && (!CanDropObjectDelegate.IsBound() || CanDropObjectDelegate.Execute(*Object));
+	}
+}
diff --git a/Engine/Plugins/Developer/Concert/ConcertUI/ConcertClientSharedSlate/Source/ConcertClientSharedSlate/Private/Replication/Editor/View/SReplicationDropArea.h b/Engine/Plugins/Developer/Concert/ConcertUI/ConcertClientSharedSlate/Source/ConcertClientSharedSlate/Private/Replication/Editor/View/SReplicationDropArea.h
new file mode 100644
--- /dev/null
+++ b/Engine/Plugins/Developer/Concert/ConcertUI/ConcertClientSharedSlate/Source/ConcertClientSharedSlate/Private/Replication/Editor/View/SReplicationDropArea.h
@@ -0,0 +1,42 @@
+﻿// Copyright Epic Games, Inc. All Rights Reserved.
+
+#pragma once
+
+#include "Replication/ClientReplicationWidgetDelegates.h"
+
+#include "Widgets/DeclarativeSyntaxSupport.h"
+#include "Widgets/SCompoundWidget.h"
+
+namespace UE::ConcertClientSharedSlate
+{
+	/** Handles drag-drop operations for the replication UI. */
+	class SReplicationDropArea : public SCompoundWidget
+	{
+	public:
+
+		SLATE_BEGIN_ARGS(SReplicationDropArea) {}
+			/** Required. Responds to objects being dropped into the replication outliner. */
+			SLATE_EVENT(FDragDropReplicatableObject, HandleDroppedObjects)
+			/** Optional. Decides whether a dragged object can be dropped. */
+			SLATE_EVENT(FCanDragDropObject, CanDropObject)
+			
+			SLATE_DEFAULT_SLOT(FArguments, Content)
+		SLATE_END_ARGS()
+
+		void Construct(const FArguments& InArgs);
+
+	private:
+
+		/** Responds to objects being dropped into the replication outliner. */
+		FDragDropReplicatableObject HandleDroppedObjectsDelegate;
+		/** Decides whether a dragged object can be dropped. */
+		FCanDragDropObject CanDropObjectDelegate;
+
+		FReply OnDragDropTarget(const FGeometry& Geometry, const FDragDropEvent& DragDropEvent) const;
+		bool CanDragDropTarget(TSharedPtr<FDragDropOperation> InOperation) const;
+
+		/** @return Whether Object can be dropped. */
+		bool CanDrop(UObject* Object) const;
+	};
+}
+
diff --git a/Engine/Plugins/Developer/Concert/ConcertUI/ConcertClientSharedSlate/Source/ConcertClientSharedSlate/Public/Replication/ClientReplicationWidgetDelegates.h b/Engine/Plugins/Developer/Concert/ConcertUI/ConcertClientSharedSlate/Source/ConcertClientSharedSlate/Public/Replication/ClientReplicationWidgetDelegates.h
new file mode 100644
--- /dev/null
+++ b/Engine/Plugins/Developer/Concert/ConcertUI/ConcertClientSharedSlate/Source/ConcertClientSharedSlate/Public/Replication/ClientReplicationWidgetDelegates.h
@@ -0,0 +1,16 @@
+﻿// Copyright Epic Games, Inc. All Rights Reserved.
+
+#pragma once
+
+#include "Containers/ContainersFwd.h"
+#include "Delegates/Delegate.h"
+
+class UObject;
+
+namespace UE::ConcertClientSharedSlate
+{
+	/** Responds to objects being dropped into the replication outliner. */
+	DECLARE_DELEGATE_OneParam(FDragDropReplicatableObject, TConstArrayView<UObject*> DroppedObjects);
+	/** Decides whether a dragged object can be dropped. */
+	DECLARE_DELEGATE_RetVal_OneParam(bool, FCanDragDropObject, UObject& Object);
+}
\ No newline at end of file
diff --git a/Engine/Plugins/Developer/Concert/ConcertUI/ConcertClientSharedSlate/Source/ConcertClientSharedSlate/Public/Replication/ClientReplicationWidgetFactories.h b/Engine/Plugins/Developer/Concert/ConcertUI/ConcertClientSharedSlate/Source/ConcertClientSharedSlate/Public/Replication/ClientReplicationWidgetFactories.h
--- a/Engine/Plugins/Developer/Concert/ConcertUI/ConcertClientSharedSlate/Source/ConcertClientSharedSlate/Public/Replication/ClientReplicationWidgetFactories.h
+++ b/Engine/Plugins/Developer/Concert/ConcertUI/ConcertClientSharedSlate/Source/ConcertClientSharedSlate/Public/Replication/ClientReplicationWidgetFactories.h
@@ -2,8 +2,9 @@
 
 #pragma once
 
+#include "ClientReplicationWidgetDelegates.h"
 #include "Replication/PropertyTreeFactory.h"
-#include "Replication/ReplicationWidgetFactories.h"
+#include "Replication/ClientReplicationWidgetFactories.h"
 #include "Templates/SharedPointer.h"
 
 struct FConcertStreamObjectAutoBindingRules;
@@ -75,4 +76,17 @@ namespace UE::ConcertClientSharedSlate
 	 * The user can toggle used filters on and off under the search bar.
 	 */
 	CONCERTCLIENTSHAREDSLATE_API TSharedRef<ConcertSharedSlate::IPropertyTreeView> CreateFilterablePropertyTreeView(FFilterablePropertyTreeViewParams Params);
+
+	struct FCreateDropTargetOutlinerWrapperParams
+	{
+		/** Required. Handles objects that were dropped. */
+		FDragDropReplicatableObject HandleDroppedObjectsDelegate;
+		/** Optional. Decides whether a dragged object can be dropped. */
+		FCanDragDropObject CanDropObjectDelegate;
+	};
+	
+	/** Returns a delegate that wraps the replication outliner with a SDropTarget: it allows users to drag-drop actors. */
+	CONCERTCLIENTSHAREDSLATE_API ConcertSharedSlate::FWrapOutlinerWidget CreateDropTargetOutlinerWrapper(
+		FCreateDropTargetOutlinerWrapperParams Params
+		);
 }
