Apply same fix to ConcertClientSharedSlate for CrashReportClientEditor as was previously done for UnrealMultiUserServer allowing it to build with -allmodules

#rnx

diff --git a/Engine/Plugins/Developer/Concert/ConcertUI/ConcertClientSharedSlate/ConcertClientSharedSlate.uplugin b/Engine/Plugins/Developer/Concert/ConcertUI/ConcertClientSharedSlate/ConcertClientSharedSlate.uplugin
--- a/Engine/Plugins/Developer/Concert/ConcertUI/ConcertClientSharedSlate/ConcertClientSharedSlate.uplugin
+++ b/Engine/Plugins/Developer/Concert/ConcertUI/ConcertClientSharedSlate/ConcertClientSharedSlate.uplugin
@@ -21,7 +21,7 @@
           "Name": "ConcertClientSharedSlate",
           "Type": "UncookedOnly",
           "LoadingPhase": "Default",
-          "ProgramAllowList": [ "UnrealRecoverySvc", "CrashReportClientEditor" ]
+          "ProgramAllowList": [ "UnrealRecoverySvc" ]
         }
 	],
 	"Plugins": [
