Fixing issue where a USTRUCT inherits from more than one USTRCT.  This isn't allowed but wasn't being checked by C# UHT.

#rnx

diff --git a/Engine/Plugins/Developer/Concert/ConcertUI/ConcertClientSharedSlate/Source/ConcertClientSharedSlate/Public/Replication/Settings/ConcertPerClassSubobjectMatchingRules.h b/Engine/Plugins/Developer/Concert/ConcertUI/ConcertClientSharedSlate/Source/ConcertClientSharedSlate/Public/Replication/Settings/ConcertPerClassSubobjectMatchingRules.h
--- a/Engine/Plugins/Developer/Concert/ConcertUI/ConcertClientSharedSlate/Source/ConcertClientSharedSlate/Public/Replication/Settings/ConcertPerClassSubobjectMatchingRules.h
+++ b/Engine/Plugins/Developer/Concert/ConcertUI/ConcertClientSharedSlate/Source/ConcertClientSharedSlate/Public/Replication/Settings/ConcertPerClassSubobjectMatchingRules.h
@@ -2,14 +2,17 @@
 
 #pragma once
 
-#include "ConcertInheritableClassOption.h"
 #include "ConcertSubobjectMatchingRules.h"
 #include "ConcertPerClassSubobjectMatchingRules.generated.h"
 
 USTRUCT()
-struct CONCERTCLIENTSHAREDSLATE_API FConcertInheritableSubobjectMatchingRules : public FConcertSubobjectMatchingRules, public FConcertInheritableClassOption
+struct CONCERTCLIENTSHAREDSLATE_API FConcertInheritableSubobjectMatchingRules : public FConcertSubobjectMatchingRules
 {
 	GENERATED_BODY()
+
+	/** Whether to inherit the settings from the base class. */
+	UPROPERTY(EditAnywhere, Category = "Config")
+	bool bInheritFromBase = true;
 };
 
 /**
