Various fixes to make unreal editor compile with IWYU

diff --git a/Engine/Plugins/Developer/Concert/ConcertUI/ConcertClientSharedSlate/Source/ConcertClientSharedSlate/Public/Replication/ClientReplicationWidgetFactories.h b/Engine/Plugins/Developer/Concert/ConcertUI/ConcertClientSharedSlate/Source/ConcertClientSharedSlate/Public/Replication/ClientReplicationWidgetFactories.h
--- a/Engine/Plugins/Developer/Concert/ConcertUI/ConcertClientSharedSlate/Source/ConcertClientSharedSlate/Public/Replication/ClientReplicationWidgetFactories.h
+++ b/Engine/Plugins/Developer/Concert/ConcertUI/ConcertClientSharedSlate/Source/ConcertClientSharedSlate/Public/Replication/ClientReplicationWidgetFactories.h
@@ -4,7 +4,6 @@
 
 #include "ClientReplicationWidgetDelegates.h"
 #include "Replication/PropertyTreeFactory.h"
-#include "Replication/ClientReplicationWidgetFactories.h"
 #include "Templates/SharedPointer.h"
 
 struct FConcertStreamObjectAutoBindingRules;
diff --git a/Engine/Plugins/Experimental/PlainProps/Source/Tests/PlainPropsUeCoreTests.cpp b/Engine/Plugins/Experimental/PlainProps/Source/Tests/PlainPropsUeCoreTests.cpp
--- a/Engine/Plugins/Experimental/PlainProps/Source/Tests/PlainPropsUeCoreTests.cpp
+++ b/Engine/Plugins/Experimental/PlainProps/Source/Tests/PlainPropsUeCoreTests.cpp
@@ -1,6 +1,6 @@
 // Copyright Epic Games, Inc. All Rights Reserved.
 
-#if WITH_TESTS
+#if WITH_TESTS && !defined(PLATFORM_COMPILER_IWYU)
 
 #include "PlainPropsBuildSchema.h"
 #include "PlainPropsCtti.h"
diff --git a/Engine/Plugins/Media/PixelStreaming/Source/PixelStreaming/PixelStreaming.Build.cs b/Engine/Plugins/Media/PixelStreaming/Source/PixelStreaming/PixelStreaming.Build.cs
--- a/Engine/Plugins/Media/PixelStreaming/Source/PixelStreaming/PixelStreaming.Build.cs
+++ b/Engine/Plugins/Media/PixelStreaming/Source/PixelStreaming/PixelStreaming.Build.cs
@@ -59,6 +59,8 @@ namespace UnrealBuildTool.Rules
 
         public PixelStreaming(ReadOnlyTargetRules Target) : base(Target)
         {
+			IWYUSupport = IWYUSupport.None;
+
             // This is so for game projects using our public headers don't have to include extra modules they might not know about.
             PublicDependencyModuleNames.AddRange(new string[]
             {
diff --git a/Engine/Plugins/Runtime/Database/ADOSupport/Source/ADOSupport/ADOSupport.Build.cs b/Engine/Plugins/Runtime/Database/ADOSupport/Source/ADOSupport/ADOSupport.Build.cs
--- a/Engine/Plugins/Runtime/Database/ADOSupport/Source/ADOSupport/ADOSupport.Build.cs
+++ b/Engine/Plugins/Runtime/Database/ADOSupport/Source/ADOSupport/ADOSupport.Build.cs
@@ -9,7 +9,7 @@ namespace UnrealBuildTool.Rules
 	{
 		public ADOSupport(ReadOnlyTargetRules Target) : base(Target)
 		{
-			if (OperatingSystem.IsWindows() && Target.Platform == UnrealTargetPlatform.Win64)
+			if (OperatingSystem.IsWindows() && Target.Platform == UnrealTargetPlatform.Win64 && !Target.bIWYU)
 			{
 				string AdoFolder = Path.Combine(Environment.GetFolderPath(Environment.SpecialFolder.CommonProgramFiles), "System", "ADO");
 				PublicSystemIncludePaths.Add(AdoFolder);
